Copying //UE4/Dev-AnimPhys to //UE4/Dev-Main (Source: //UE4/Dev-AnimPhys @ 3624379)

diff --git a/Engine/Source/Editor/PhysicsAssetEditor/Private/PhysicsAssetEditorActions.cpp b/Engine/Source/Editor/PhysicsAssetEditor/Private/PhysicsAssetEditorActions.cpp
new file mode 100644
--- /dev/null
+++ b/Engine/Source/Editor/PhysicsAssetEditor/Private/PhysicsAssetEditorActions.cpp
@@ -0,0 +1,101 @@
+// Copyright 1998-2017 Epic Games, Inc. All Rights Reserved.
+
+#include "PhysicsAssetEditorActions.h"
+
+#define LOCTEXT_NAMESPACE "PhysicsAssetEditorCommands"
+
+void FPhysicsAssetEditorCommands::RegisterCommands()
+{
+	UI_COMMAND(RegenerateBodies, "Regenerate Bodies", "Regenerates the selected bodies using the current generation settings (see the Tools tab)", EUserInterfaceActionType::Button, FInputChord());
+	UI_COMMAND(AddBodies, "Add/Replace Bodies", "Adds or replaces bodies for the selected bones using the current generation settings (see the Tools tab)", EUserInterfaceActionType::Button, FInputChord());
+
+	UI_COMMAND(ApplyPhysicalMaterial, "Physical Material", "Apply a Physical Material to All Bodies", EUserInterfaceActionType::Button, FInputChord());
+	UI_COMMAND(CopyProperties, "Copy Properties", "Copy Properties: Copy Properties Of Currently Selected Object To Next Selected Object", EUserInterfaceActionType::Button, FInputChord(EModifierKey::Control, EKeys::C));
+	UI_COMMAND(PasteProperties, "Paste Properties", "Paste Properties: Copy Properties Of Currently Selected Object To Next Selected Object", EUserInterfaceActionType::Button, FInputChord(EModifierKey::Control, EKeys::V));
+	UI_COMMAND(RepeatLastSimulation, "Simulate", "Previews Physics Simulation", EUserInterfaceActionType::RadioButton, FInputChord());
+	UI_COMMAND(SimulationNoGravity, "No Gravity Simulation", "Run Physics Simulation without gravity. Use this to debug issues with your ragdoll. If the setup is correct, the asset should not move!", EUserInterfaceActionType::ToggleButton, FInputChord());
+	UI_COMMAND(SelectedSimulation, "Simulate Selected", "Run Physics Simulation on selected objects. Use this to tune  specific parts of your ragdoll.", EUserInterfaceActionType::RadioButton, FInputChord(EModifierKey::Alt | EModifierKey::Shift, EKeys::Enter));
+	UI_COMMAND(SimulationAll, "Simulate", "Run Physics Simulation on all objects.", EUserInterfaceActionType::RadioButton, FInputChord(EModifierKey::Alt, EKeys::Enter));
+
+	UI_COMMAND(MeshRenderingMode_Solid, "Solid", "Solid Mesh Rendering Mode", EUserInterfaceActionType::RadioButton, FInputChord());
+	UI_COMMAND(MeshRenderingMode_Wireframe, "Wireframe", "Wireframe Mesh Rendering Mode", EUserInterfaceActionType::RadioButton, FInputChord());
+	UI_COMMAND(MeshRenderingMode_None, "None", "No Mesh Rendering Mode", EUserInterfaceActionType::RadioButton, FInputChord());
+	UI_COMMAND(CollisionRenderingMode_Solid, "Solid", "Solid Collision Rendering Mode", EUserInterfaceActionType::RadioButton, FInputChord());
+	UI_COMMAND(CollisionRenderingMode_Wireframe, "Wireframe", "Wireframe Collision Rendering Mode", EUserInterfaceActionType::RadioButton, FInputChord());
+	UI_COMMAND(CollisionRenderingMode_None, "None", "No Collision Rendering Mode", EUserInterfaceActionType::RadioButton, FInputChord());
+	UI_COMMAND(ConstraintRenderingMode_None, "None", "No Constraint Rendering Mode", EUserInterfaceActionType::RadioButton, FInputChord());
+	UI_COMMAND(ConstraintRenderingMode_AllPositions, "All Positions", "All Positions Constraint Rendering Mode", EUserInterfaceActionType::RadioButton, FInputChord());
+	UI_COMMAND(ConstraintRenderingMode_AllLimits, "All Limits", "All Limits Constraint Rendering Mode", EUserInterfaceActionType::RadioButton, FInputChord());
+
+	UI_COMMAND(MeshRenderingMode_Simulation_Solid, "Solid", "Solid Mesh Rendering Mode", EUserInterfaceActionType::RadioButton, FInputChord());
+	UI_COMMAND(MeshRenderingMode_Simulation_Wireframe, "Wireframe", "Wireframe Mesh Rendering Mode", EUserInterfaceActionType::RadioButton, FInputChord());
+	UI_COMMAND(MeshRenderingMode_Simulation_None, "None", "No Mesh Rendering Mode", EUserInterfaceActionType::RadioButton, FInputChord());
+	UI_COMMAND(CollisionRenderingMode_Simulation_Solid, "Solid", "Solid Collision Rendering Mode", EUserInterfaceActionType::RadioButton, FInputChord());
+	UI_COMMAND(CollisionRenderingMode_Simulation_Wireframe, "Wireframe", "Wireframe Collision Rendering Mode", EUserInterfaceActionType::RadioButton, FInputChord());
+	UI_COMMAND(CollisionRenderingMode_Simulation_None, "None", "No Collision Rendering Mode", EUserInterfaceActionType::RadioButton, FInputChord());
+	UI_COMMAND(ConstraintRenderingMode_Simulation_None, "None", "No Constraint Rendering Mode", EUserInterfaceActionType::RadioButton, FInputChord());
+	UI_COMMAND(ConstraintRenderingMode_Simulation_AllPositions, "All Positions", "All Positions Constraint Rendering Mode", EUserInterfaceActionType::RadioButton, FInputChord());
+	UI_COMMAND(ConstraintRenderingMode_Simulation_AllLimits, "All Limits", "All Limits Constraint Rendering Mode", EUserInterfaceActionType::RadioButton, FInputChord());
+
+	UI_COMMAND(RenderOnlySelectedSolid, "Only Selected Solid", "Only render selected collision as 'solid'", EUserInterfaceActionType::ToggleButton, FInputChord());
+	UI_COMMAND(DrawConstraintsAsPoints, "Draw Constraints As Points", "Draw Constraints As Points", EUserInterfaceActionType::ToggleButton, FInputChord());
+	UI_COMMAND(ToggleMassProperties, "Mass Properties", "Show Mass Properties For Bodies When Simulation Is Enabled", EUserInterfaceActionType::ToggleButton, FInputChord());
+	UI_COMMAND(DisableCollision, "Disable Collision", "Disable collision between the currently selected bodies", EUserInterfaceActionType::Button, FInputChord(EKeys::RightBracket));
+	UI_COMMAND(DisableCollisionAll, "Disable Collision All", "Disable collision between the currently selected bodies and all bodies", EUserInterfaceActionType::Button, FInputChord(EModifierKey::Control, EKeys::RightBracket));
+	UI_COMMAND(EnableCollision, "Enable Collision", "Enable collision between the currently selected bodies", EUserInterfaceActionType::Button, FInputChord(EKeys::LeftBracket));
+	UI_COMMAND(EnableCollisionAll, "Enable Collision All", "Enable collision between the currently selected bodies and all bodies", EUserInterfaceActionType::Button, FInputChord(EModifierKey::Control, EKeys::LeftBracket));
+	UI_COMMAND(WeldToBody, "Weld", "Weld Body: Weld Currently Selected Bodies", EUserInterfaceActionType::Button, FInputChord());
+	UI_COMMAND(AddSphere, "Add Sphere", "Add Sphere To Selected Bone", EUserInterfaceActionType::Button, FInputChord());
+	UI_COMMAND(AddSphyl, "Add Capsule", "Add Capsule To Selected Bone", EUserInterfaceActionType::Button, FInputChord());
+	UI_COMMAND(AddBox, "Add Box", "Add Box To Selected Bone", EUserInterfaceActionType::Button, FInputChord());
+	UI_COMMAND(DeletePrimitive, "Delete", "Delete Selected Primitive(s)", EUserInterfaceActionType::Button, FInputChord());
+	UI_COMMAND(DuplicatePrimitive, "Duplicate", "Duplicate Selected Primitive(s)", EUserInterfaceActionType::Button, FInputChord());
+	UI_COMMAND(ResetConstraint, "Reset", "Reset Constraint", EUserInterfaceActionType::Button, FInputChord());
+	UI_COMMAND(SnapConstraint, "Snap", "Snap Constraint To Bone", EUserInterfaceActionType::Button, FInputChord());
+	UI_COMMAND(ConvertToBallAndSocket, "To Ball & Socket", "Convert Selected Constraint To Ball-And-Socket", EUserInterfaceActionType::Button, FInputChord());
+	UI_COMMAND(ConvertToHinge, "To Hinge", "Convert Selected Constraint To Hinge", EUserInterfaceActionType::Button, FInputChord());
+	UI_COMMAND(ConvertToPrismatic, "To Prismatic", "Convert Selected Constraint To Prismatic", EUserInterfaceActionType::Button, FInputChord());
+	UI_COMMAND(ConvertToSkeletal, "To Skeletal", "Convert Selected Constraint To Skeletal", EUserInterfaceActionType::Button, FInputChord());
+	UI_COMMAND(DeleteConstraint, "Delete", "Delete Selected Constraint", EUserInterfaceActionType::Button, FInputChord());
+	UI_COMMAND(ShowSkeleton, "Skeleton", "Show Skeleton", EUserInterfaceActionType::ToggleButton, FInputChord());
+	UI_COMMAND(MakeBodyKinematic, "Kinematic", "Make Body Kinematic", EUserInterfaceActionType::RadioButton, FInputChord());
+	UI_COMMAND(MakeBodySimulated, "Simulated", "Make Body Simulated", EUserInterfaceActionType::RadioButton, FInputChord());
+	UI_COMMAND(MakeBodyDefault, "Default", "Reset This Body To Default", EUserInterfaceActionType::RadioButton, FInputChord());
+	UI_COMMAND(KinematicAllBodiesBelow, "Set All Bodies Below To Kinematic", "Set All Bodies Below To Kinematic", EUserInterfaceActionType::Button, FInputChord());
+	UI_COMMAND(SimulatedAllBodiesBelow, "Set All Bodies Below To Simulated", "Set All Bodies Below To Simulated", EUserInterfaceActionType::Button, FInputChord());
+	UI_COMMAND(MakeAllBodiesBelowDefault, "Reset All Bodies Below To Default", "Reset All Bodies Below To Default", EUserInterfaceActionType::Button, FInputChord());
+	UI_COMMAND(DeleteBody, "Delete", "Delete Selected Body", EUserInterfaceActionType::Button, FInputChord());
+	UI_COMMAND(DeleteAllBodiesBelow, "Delete All Bodies Below", "Delete All Bodies Below", EUserInterfaceActionType::Button, FInputChord());
+	UI_COMMAND(SelectAllBodies, "Select All Bodies", "Select All Bodies", EUserInterfaceActionType::Button, FInputChord(EModifierKey::Control, EKeys::A));
+	UI_COMMAND(SelectAllConstraints, "Select All Constraints", "Select All Constraints", EUserInterfaceActionType::Button, FInputChord(EModifierKey::Control | EModifierKey::Shift, EKeys::A));
+	UI_COMMAND(DeselectAll, "Deselect All Bodies", "Deselect All", EUserInterfaceActionType::Button, FInputChord(EKeys::Escape));
+	UI_COMMAND(Mirror, "Mirror", "Finds the body on the other side and duplicates constraint and body", EUserInterfaceActionType::Button, FInputChord(EKeys::M));
+
+	UI_COMMAND(NewPhysicalAnimationProfile, "New", "Creates a new physical animation profile", EUserInterfaceActionType::Button, FInputChord());
+	UI_COMMAND(DuplicatePhysicalAnimationProfile, "Duplicate", "Duplicates the current physical animation profile", EUserInterfaceActionType::Button, FInputChord());
+	UI_COMMAND(DeleteCurrentPhysicalAnimationProfile, "Delete", "Deletes the current physical animation profile", EUserInterfaceActionType::Button, FInputChord());
+	UI_COMMAND(AddBodyToPhysicalAnimationProfile, "Assign", "Assigns the selected bodies to the current physical animation profile", EUserInterfaceActionType::Button, FInputChord());
+	UI_COMMAND(RemoveBodyFromPhysicalAnimationProfile, "Unassign", "Unassigns the selected bodies from the current physical animation profile", EUserInterfaceActionType::Button, FInputChord());
+	UI_COMMAND(NewConstraintProfile, "New", "Creates a new constraint profile", EUserInterfaceActionType::Button, FInputChord());
+	UI_COMMAND(DuplicateConstraintProfile, "Duplicate", "Duplicates the current constraint profile", EUserInterfaceActionType::Button, FInputChord());
+	UI_COMMAND(DeleteCurrentConstraintProfile, "Delete", "Deletes the current constraint profile", EUserInterfaceActionType::Button, FInputChord());
+	UI_COMMAND(AddConstraintToCurrentConstraintProfile, "Assign", "Assigns the selected constraints to the current constraint profile", EUserInterfaceActionType::Button, FInputChord());
+	UI_COMMAND(RemoveConstraintFromCurrentConstraintProfile, "Unassign", "Unassigns the selected constraints from the current constraint profile", EUserInterfaceActionType::Button, FInputChord());
+
+	UI_COMMAND(ShowBodies, "Show Bodies", "Display bodies in the tree view. Bodies are a collection of primitive shapes used for collision.", EUserInterfaceActionType::ToggleButton, FInputChord());
+	UI_COMMAND(ShowConstraints, "Show Constraints", "Display constraints in the tree view. Constraints are used to control how bodies can move in relation to one another.", EUserInterfaceActionType::ToggleButton, FInputChord());
+	UI_COMMAND(ShowPrimitives, "Show Primitives", "Display primitive shapes (sphere, boxes, capsules etc.) in the tree view", EUserInterfaceActionType::ToggleButton, FInputChord());
+
+	/** As two commands cannot have the same key; this command wraps both 
+	  * DeletePrimitive and DeleteConstraint so the user can delete whatever is selected 
+	  */
+	UI_COMMAND(DeleteSelected, "Delete selected primitive or constraint", "", EUserInterfaceActionType::Button, FInputChord(EKeys::Platform_Delete));
+	UI_COMMAND(CycleConstraintOrientation, "Cycle Constraint Orientation", "Cycles constraint orientation between different cardinal axes", EUserInterfaceActionType::Button, FInputChord(EKeys::Q));
+	UI_COMMAND(CycleConstraintActive, "Cycle Active Constraint", "Cycles whether each constraint axis is active in isolation", EUserInterfaceActionType::Button, FInputChord(EKeys::Four));
+	UI_COMMAND(ToggleSwing1, "Lock Swing 1", "Set swing 1 to be locked or limited", EUserInterfaceActionType::ToggleButton, FInputChord(EKeys::One));
+	UI_COMMAND(ToggleSwing2, "Lock Swing 2", "Set swing 2 to be locked or limited", EUserInterfaceActionType::ToggleButton, FInputChord(EKeys::Two));
+	UI_COMMAND(ToggleTwist, "Lock Twist", "Set twist to be locked or limited", EUserInterfaceActionType::ToggleButton, FInputChord(EKeys::Three));
+	UI_COMMAND(FocusOnSelection, "Focus the viewport on the current selection", "", EUserInterfaceActionType::Button, FInputChord(EKeys::F));
+}
+
+#undef LOCTEXT_NAMESPACE
diff --git a/Engine/Source/Editor/PhysicsAssetEditor/Private/PhysicsAssetEditorAnimInstanceProxy.cpp b/Engine/Source/Editor/PhysicsAssetEditor/Private/PhysicsAssetEditorAnimInstanceProxy.cpp
new file mode 100644
--- /dev/null
+++ b/Engine/Source/Editor/PhysicsAssetEditor/Private/PhysicsAssetEditorAnimInstanceProxy.cpp
@@ -0,0 +1,30 @@
+// Copyright 1998-2016 Epic Games, Inc. All Rights Reserved.
+
+#include "PhysicsAssetEditorAnimInstanceProxy.h"
+#include "PhysicsAssetEditorAnimInstance.h"
+//#include "AnimNode_RigidBody.h"
+
+void FPhysicsAssetEditorAnimInstanceProxy::Initialize(UAnimInstance* InAnimInstance)
+{
+	FAnimPreviewInstanceProxy::Initialize(InAnimInstance);
+	ConstructNodes();
+}
+
+void FPhysicsAssetEditorAnimInstanceProxy::ConstructNodes()
+{
+	//LocalToComponentSpace.LocalPose.SetLinkNode(&SequencePlayerNode);
+	//RagdollNode.ComponentPose.SetLinkNode(&LocalToComponentSpace);
+	//ComponentToLocalSpace.ComponentPose.SetLinkNode(&RagdollNode);
+	//ComponentToLocalSpace.ComponentPose.SetLinkNode(&LocalToComponentSpace);
+	
+	//RagdollNode.SimulationSpace = ESimulationSpace::WorldSpace;
+}
+
+void FPhysicsAssetEditorAnimInstanceProxy::GetCustomNodes(TArray<FAnimNode_Base*>& OutNodes)
+{
+	//OutNodes.Add(&RagdollNode);
+}
+
+FPhysicsAssetEditorAnimInstanceProxy::~FPhysicsAssetEditorAnimInstanceProxy()
+{
+}
\ No newline at end of file
diff --git a/Engine/Source/Editor/PhysicsAssetEditor/Private/PhysicsAssetEditorAnimInstanceProxy.h b/Engine/Source/Editor/PhysicsAssetEditor/Private/PhysicsAssetEditorAnimInstanceProxy.h
new file mode 100644
--- /dev/null
+++ b/Engine/Source/Editor/PhysicsAssetEditor/Private/PhysicsAssetEditorAnimInstanceProxy.h
@@ -0,0 +1,43 @@
+// Copyright 1998-2016 Epic Games, Inc. All Rights Reserved.
+
+#pragma once
+
+#include "AnimPreviewInstance.h"
+//#include "AnimNode_RigidBody.h"
+#include "Animation/AnimNodeSpaceConversions.h"
+#include "Animation/AnimNode_SequencePlayer.h"
+#include "PhysicsAssetEditorAnimInstanceProxy.generated.h"
+
+class UAnimSequence;
+
+/** Proxy override for this UAnimInstance-derived class */
+USTRUCT()
+struct FPhysicsAssetEditorAnimInstanceProxy : public FAnimPreviewInstanceProxy
+{
+	GENERATED_BODY()
+
+public:
+	FPhysicsAssetEditorAnimInstanceProxy()
+	{
+	}
+
+	FPhysicsAssetEditorAnimInstanceProxy(UAnimInstance* InAnimInstance)
+		: FAnimPreviewInstanceProxy(InAnimInstance)
+	{
+	}
+
+	virtual ~FPhysicsAssetEditorAnimInstanceProxy();
+
+	virtual void Initialize(UAnimInstance* InAnimInstance) override;
+	virtual void GetCustomNodes(TArray<FAnimNode_Base*>& OutNodes) override;
+private:
+	void ConstructNodes();
+
+	//FAnimNode_RigidBody RagdollNode;
+// 	FAnimNode_ConvertComponentToLocalSpace ComponentToLocalSpace;
+// 	FAnimNode_SequencePlayer SequencePlayerNode;
+// 	FAnimNode_ConvertLocalToComponentSpace LocalToComponentSpace;
+	
+	void InitAnimTrack(UAnimSequenceBase* InAnimSequence, int32 SequenceId);
+	void EnsureAnimTrack(UAnimSequenceBase* InAnimSequence, int32 SequenceId);
+};
\ No newline at end of file
diff --git a/Engine/Source/Editor/PhysicsAssetEditor/Private/PhysicsAssetEditorEditMode.cpp b/Engine/Source/Editor/PhysicsAssetEditor/Private/PhysicsAssetEditorEditMode.cpp
new file mode 100644
--- /dev/null
+++ b/Engine/Source/Editor/PhysicsAssetEditor/Private/PhysicsAssetEditorEditMode.cpp
@@ -0,0 +1,935 @@
+// Copyright 1998-2017 Epic Games, Inc. All Rights Reserved.
+
+#include "PhysicsAssetEditorEditMode.h"
+#include "PhysicsAssetEditorSkeletalMeshComponent.h"
+#include "CanvasItem.h"
+#include "CanvasTypes.h"
+#include "AssetEditorModeManager.h"
+#include "EngineUtils.h"
+#include "PhysicsAssetEditorSharedData.h"
+#include "PhysicsEngine/PhysicsAsset.h"
+#include "PhysicsEngine/PhysicsConstraintTemplate.h"
+#include "PhysicsEngine/PhysicsSettings.h"
+#include "Preferences/PhysicsAssetEditorOptions.h"
+#include "IPersonaPreviewScene.h"
+#include "PhysicsAssetEditor.h"
+#include "PhysicsAssetEditorHitProxies.h"
+#include "PhysicsEngine/PhysicsHandleComponent.h"
+#include "DrawDebugHelpers.h"
+#include "SEditorViewport.h"
+#include "IPersonaToolkit.h"
+#include "MultiBoxBuilder.h"
+#include "SlateApplication.h"
+
+#define LOCTEXT_NAMESPACE "PhysicsAssetEditorEditMode"
+
+FName FPhysicsAssetEditorEditMode::ModeName("PhysicsAssetEditor");
+
+FPhysicsAssetEditorEditMode::FPhysicsAssetEditorEditMode()
+	: MinPrimSize(0.5f)
+	, PhysicsAssetEditor_TranslateSpeed(0.25f)
+	, PhysicsAssetEditor_RotateSpeed(1.0f * (PI / 180.0f))
+	, PhysicsAssetEditor_LightRotSpeed(0.22f)
+	, SimGrabCheckDistance(5000.0f)
+	, SimHoldDistanceChangeDelta(20.0f)
+	, SimMinHoldDistance(10.0f)
+	, SimGrabMoveSpeed(1.0f)
+{
+	// Disable grid drawing for this mode as the viewport handles this
+	bDrawGrid = false;
+
+	PhysicsAssetEditorFont = GEngine->GetSmallFont();
+	check(PhysicsAssetEditorFont);
+}
+
+bool FPhysicsAssetEditorEditMode::GetCameraTarget(FSphere& OutTarget) const
+{
+	bool bHandled = false;
+
+	float NumSelected = 0.0f;
+	FBox Bounds(ForceInit);
+	for (int32 BodyIndex = 0; BodyIndex < SharedData->SelectedBodies.Num(); ++BodyIndex)
+	{
+		FPhysicsAssetEditorSharedData::FSelection& SelectedObject = SharedData->SelectedBodies[BodyIndex];
+		int32 BoneIndex = SharedData->EditorSkelComp->GetBoneIndex(SharedData->PhysicsAsset->SkeletalBodySetups[SelectedObject.Index]->BoneName);
+		UBodySetup* BodySetup = SharedData->PhysicsAsset->SkeletalBodySetups[SelectedObject.Index];
+		const FKAggregateGeom& AggGeom = BodySetup->AggGeom;
+
+		FTransform BoneTM = SharedData->EditorSkelComp->GetBoneTransform(BoneIndex);
+		const float Scale = BoneTM.GetScale3D().GetAbsMax();
+		BoneTM.RemoveScaling();
+
+		if (SelectedObject.PrimitiveType == EAggCollisionShape::Sphere)
+		{
+			Bounds += AggGeom.SphereElems[SelectedObject.PrimitiveIndex].CalcAABB(BoneTM, Scale);
+		}
+		else if (SelectedObject.PrimitiveType == EAggCollisionShape::Box)
+		{
+			Bounds += AggGeom.BoxElems[SelectedObject.PrimitiveIndex].CalcAABB(BoneTM, Scale);
+		}
+		else if (SelectedObject.PrimitiveType == EAggCollisionShape::Sphyl)
+		{
+			Bounds += AggGeom.SphylElems[SelectedObject.PrimitiveIndex].CalcAABB(BoneTM, Scale);
+		}
+		else if (SelectedObject.PrimitiveType == EAggCollisionShape::Convex)
+		{
+			Bounds += AggGeom.ConvexElems[SelectedObject.PrimitiveIndex].CalcAABB(BoneTM, BoneTM.GetScale3D());
+		}
+
+		bHandled = true;
+	}
+
+	for (int32 ConstraintIndex = 0; ConstraintIndex < SharedData->SelectedConstraints.Num(); ++ConstraintIndex)
+	{
+		FPhysicsAssetEditorSharedData::FSelection& SelectedObject = SharedData->SelectedConstraints[ConstraintIndex];
+		Bounds += SharedData->GetConstraintWorldTM(&SelectedObject, EConstraintFrame::Frame2).GetLocation();
+
+		bHandled = true;
+	}
+
+	OutTarget.Center = Bounds.GetCenter();
+	OutTarget.W = Bounds.GetExtent().Size();	// @TODO: calculate correct bounds
+
+	return bHandled;
+}
+
+IPersonaPreviewScene& FPhysicsAssetEditorEditMode::GetAnimPreviewScene() const
+{
+	return *static_cast<IPersonaPreviewScene*>(static_cast<FAssetEditorModeManager*>(Owner)->GetPreviewScene());
+}
+
+void FPhysicsAssetEditorEditMode::GetOnScreenDebugInfo(TArray<FText>& OutDebugInfo) const
+{
+
+}
+
+bool FPhysicsAssetEditorEditMode::StartTracking(FEditorViewportClient* InViewportClient, FViewport* InViewport)
+{
+	const EAxisList::Type CurrentAxis = InViewportClient->GetCurrentWidgetAxis();
+	if(CurrentAxis != EAxisList::None)
+	{
+		// If releasing the mouse button, check we are done manipulating
+		check(!SharedData->bManipulating);
+		if (SharedData->GetSelectedBody())
+		{
+			GEditor->BeginTransaction(NSLOCTEXT("UnrealEd", "MoveElement", "Move Element"));
+			for (int32 i = 0; i<SharedData->SelectedBodies.Num(); ++i)
+			{
+				SharedData->PhysicsAsset->SkeletalBodySetups[SharedData->SelectedBodies[i].Index]->Modify();
+				SharedData->SelectedBodies[i].ManipulateTM = FTransform::Identity;
+			}
+
+			SharedData->bManipulating = true;
+		}
+
+		if (SharedData->GetSelectedConstraint())
+		{
+			GEditor->BeginTransaction(NSLOCTEXT("UnrealEd", "MoveConstraint", "Move Constraint"));
+
+			const int32 Count = SharedData->SelectedConstraints.Num();
+			StartManRelConTM.SetNumUninitialized(Count);
+			StartManParentConTM.SetNumUninitialized(Count);
+			StartManChildConTM.SetNumUninitialized(Count);
+
+			for (int32 i = 0; i < SharedData->SelectedConstraints.Num(); ++i)
+			{
+				FPhysicsAssetEditorSharedData::FSelection& Constraint = SharedData->SelectedConstraints[i];
+				SharedData->PhysicsAsset->ConstraintSetup[Constraint.Index]->Modify();
+				Constraint.ManipulateTM = FTransform::Identity;
+
+				const FTransform WParentFrame = SharedData->GetConstraintWorldTM(&Constraint, EConstraintFrame::Frame2);
+				const FTransform WChildFrame = SharedData->GetConstraintWorldTM(&Constraint, EConstraintFrame::Frame1);
+				const UPhysicsConstraintTemplate* Setup = SharedData->PhysicsAsset->ConstraintSetup[Constraint.Index];
+
+				StartManRelConTM[i] = WChildFrame * WParentFrame.Inverse();
+				StartManParentConTM[i] = Setup->DefaultInstance.GetRefFrame(EConstraintFrame::Frame2);
+				StartManChildConTM[i] = Setup->DefaultInstance.GetRefFrame(EConstraintFrame::Frame1);
+			}
+
+			SharedData->bManipulating = true;
+		}
+	}
+	
+	return SharedData->bManipulating;
+}
+
+bool FPhysicsAssetEditorEditMode::EndTracking(FEditorViewportClient* InViewportClient, FViewport* InViewport)
+{
+	if (SharedData->bManipulating)
+	{
+		SharedData->bManipulating = false;
+
+		for (int32 i = 0; i<SharedData->SelectedBodies.Num(); ++i)
+		{
+			FPhysicsAssetEditorSharedData::FSelection & SelectedObject = SharedData->SelectedBodies[i];
+			UBodySetup* BodySetup = SharedData->PhysicsAsset->SkeletalBodySetups[SelectedObject.Index];
+
+			FKAggregateGeom* AggGeom = &BodySetup->AggGeom;
+
+
+			if (SelectedObject.PrimitiveType == EAggCollisionShape::Convex)
+			{
+				FKConvexElem& Convex = AggGeom->ConvexElems[SelectedObject.PrimitiveIndex];
+				Convex.SetTransform(SelectedObject.ManipulateTM * Convex.GetTransform());
+
+				BodySetup->InvalidatePhysicsData();
+				BodySetup->CreatePhysicsMeshes();
+			}
+		}
+
+		GEditor->EndTransaction();
+		SharedData->RefreshPhysicsAssetChange(SharedData->PhysicsAsset);
+		InViewport->Invalidate();
+
+		return true;
+	}
+
+	return false;
+}
+
+bool FPhysicsAssetEditorEditMode::InputKey(FEditorViewportClient* InViewportClient, FViewport* InViewport, FKey Key, EInputEvent Event)
+{
+	int32 HitX = InViewport->GetMouseX();
+	int32 HitY = InViewport->GetMouseY();
+	bool bCtrlDown = InViewport->KeyState(EKeys::LeftControl) || InViewport->KeyState(EKeys::RightControl);
+	bool bShiftDown = InViewport->KeyState(EKeys::LeftShift) || InViewport->KeyState(EKeys::RightShift);
+
+	bool bHandled = false;
+	if (SharedData->bRunningSimulation)
+	{
+		if (Key == EKeys::RightMouseButton || Key == EKeys::LeftMouseButton)
+		{
+			if (Event == IE_Pressed)
+			{
+				bHandled = SimMousePress(InViewportClient, Key);
+			}
+			else if (Event == IE_Released)
+			{
+				bHandled = SimMouseRelease();
+			}
+			else
+			{
+				// Handle repeats/double clicks etc. so we dont fall through
+				bHandled = true;
+			}
+		}
+		else if (Key == EKeys::MouseScrollUp)
+		{
+			bHandled = SimMouseWheelUp(InViewportClient);
+		}
+		else if (Key == EKeys::MouseScrollDown)
+		{
+			bHandled = SimMouseWheelDown(InViewportClient);
+		}
+		else if (InViewportClient->IsFlightCameraActive())
+		{
+			// If the flight camera is active (user is looking or moving around the scene)
+			// consume the event so hotkeys don't fire.
+			bHandled = true;
+		}
+	}
+
+	if (!bHandled)
+	{
+		bHandled = IPersonaEditMode::InputKey(InViewportClient, InViewport, Key, Event);
+	}
+
+	if (bHandled)
+	{
+		InViewportClient->Invalidate();
+	}
+
+	return bHandled;
+}
+
+bool FPhysicsAssetEditorEditMode::InputAxis(FEditorViewportClient* InViewportClient, FViewport* InViewport, int32 ControllerId, FKey Key, float Delta, float DeltaTime)
+{
+	bool bHandled = false;
+	// If we are 'manipulating' don't move the camera but do something else with mouse input.
+	if (SharedData->bManipulating)
+	{
+		bool bCtrlDown = InViewport->KeyState(EKeys::LeftControl) || InViewport->KeyState(EKeys::RightControl);
+
+		if (SharedData->bRunningSimulation)
+		{
+			if (Key == EKeys::MouseX)
+			{
+				SimMouseMove(InViewportClient, Delta, 0.0f);
+			}
+			else if (Key == EKeys::MouseY)
+			{
+				SimMouseMove(InViewportClient, 0.0f, Delta);
+			}
+			bHandled = true;
+		}
+	}
+
+	if (!bHandled)
+	{
+		bHandled = IPersonaEditMode::InputAxis(InViewportClient, InViewport, ControllerId, Key, Delta, DeltaTime);
+	}
+
+	InViewportClient->Invalidate();
+
+	return bHandled;
+}
+
+bool FPhysicsAssetEditorEditMode::InputDelta(FEditorViewportClient* InViewportClient, FViewport* InViewport, FVector& InDrag, FRotator& InRot, FVector& InScale)
+{
+	bool bHandled = false;
+	const EAxisList::Type CurrentAxis = InViewportClient->GetCurrentWidgetAxis();
+	if (SharedData->bManipulating && CurrentAxis != EAxisList::None)
+	{
+		for (int32 i = 0; i< SharedData->SelectedBodies.Num(); ++i)
+		{
+			FPhysicsAssetEditorSharedData::FSelection & SelectedObject = SharedData->SelectedBodies[i];
+			if (SharedData->bManipulating)
+			{
+				float BoneScale = 1.f;
+
+				int32 BoneIndex = SharedData->EditorSkelComp->GetBoneIndex(SharedData->PhysicsAsset->SkeletalBodySetups[SelectedObject.Index]->BoneName);
+
+				FTransform BoneTM = SharedData->EditorSkelComp->GetBoneTransform(BoneIndex);
+				BoneScale = BoneTM.GetScale3D().GetAbsMax();
+				BoneTM.RemoveScaling();
+
+				SelectedObject.WidgetTM = SharedData->EditorSkelComp->GetPrimitiveTransform(BoneTM, SelectedObject.Index, SelectedObject.PrimitiveType, SelectedObject.PrimitiveIndex, BoneScale);
+
+				if(InViewportClient->GetWidgetMode() == FWidget::WM_Translate || InViewportClient->GetWidgetMode() == FWidget::WM_Rotate)
+				{
+					if (InViewportClient->GetWidgetMode() == FWidget::WM_Translate)
+					{
+						FVector Dir = SelectedObject.WidgetTM.InverseTransformVector(InDrag.GetSafeNormal());
+						FVector DragVec = Dir * InDrag.Size() / BoneScale;
+						SelectedObject.ManipulateTM.AddToTranslation(DragVec);
+					}
+					else if (InViewportClient->GetWidgetMode() == FWidget::WM_Rotate)
+					{
+						FVector Axis;
+						float Angle;
+						InRot.Quaternion().ToAxisAndAngle(Axis, Angle);
+
+						Axis = SelectedObject.WidgetTM.InverseTransformVectorNoScale(Axis);
+
+						const FQuat Start = SelectedObject.ManipulateTM.GetRotation();
+						const FQuat Delta = FQuat(Axis, Angle);
+						const FQuat Result = Delta * Start;
+
+						SelectedObject.ManipulateTM = FTransform(Result);
+					}
+				
+					UBodySetup* BodySetup = SharedData->PhysicsAsset->SkeletalBodySetups[SelectedObject.Index];
+					FKAggregateGeom* AggGeom = &BodySetup->AggGeom;
+
+					// for all but convex shapes, we apply straight away
+					if (SelectedObject.PrimitiveType == EAggCollisionShape::Sphere)
+					{
+						AggGeom->SphereElems[SelectedObject.PrimitiveIndex].Center = (SelectedObject.ManipulateTM * AggGeom->SphereElems[SelectedObject.PrimitiveIndex].GetTransform()).GetLocation();
+						SelectedObject.ManipulateTM.SetIdentity();
+					}
+					else if (SelectedObject.PrimitiveType == EAggCollisionShape::Box)
+					{
+						AggGeom->BoxElems[SelectedObject.PrimitiveIndex].SetTransform(SelectedObject.ManipulateTM * AggGeom->BoxElems[SelectedObject.PrimitiveIndex].GetTransform());
+						SelectedObject.ManipulateTM.SetIdentity();
+					}
+					else if (SelectedObject.PrimitiveType == EAggCollisionShape::Sphyl)
+					{
+						AggGeom->SphylElems[SelectedObject.PrimitiveIndex].SetTransform(SelectedObject.ManipulateTM * AggGeom->SphylElems[SelectedObject.PrimitiveIndex].GetTransform());
+						SelectedObject.ManipulateTM.SetIdentity();
+					}
+				}
+				else if (InViewportClient->GetWidgetMode() == FWidget::WM_Scale)
+				{
+					ModifyPrimitiveSize(SelectedObject.Index, SelectedObject.PrimitiveType, SelectedObject.PrimitiveIndex, InScale);
+				}
+
+				bHandled = true;
+			}
+		}
+
+		for (int32 i = 0; i<SharedData->SelectedConstraints.Num(); ++i)
+		{
+			FPhysicsAssetEditorSharedData::FSelection & SelectedObject = SharedData->SelectedConstraints[i];
+			if (SharedData->bManipulating)
+			{
+				float BoneScale = 1.f;
+				SelectedObject.WidgetTM = SharedData->GetConstraintMatrix(SelectedObject.Index, EConstraintFrame::Frame2, 1.f);
+
+				if (InViewportClient->GetWidgetMode() == FWidget::WM_Translate)
+				{
+					FVector Dir = SelectedObject.WidgetTM.InverseTransformVector(InDrag.GetSafeNormal());
+					FVector DragVec = Dir * InDrag.Size() / BoneScale;
+					SelectedObject.ManipulateTM.AddToTranslation(DragVec);
+				}
+				else if (InViewportClient->GetWidgetMode() == FWidget::WM_Rotate)
+				{
+					FVector Axis;
+					float Angle;
+					InRot.Quaternion().ToAxisAndAngle(Axis, Angle);
+
+					Axis = SelectedObject.WidgetTM.InverseTransformVectorNoScale(Axis);
+
+					const FQuat Start = SelectedObject.ManipulateTM.GetRotation();
+					const FQuat Delta = FQuat(Axis, Angle);
+					const FQuat Result = Delta * Start;
+
+					SelectedObject.ManipulateTM = FTransform(Result);
+				}
+
+				UPhysicsConstraintTemplate* ConstraintSetup = SharedData->PhysicsAsset->ConstraintSetup[SelectedObject.Index];
+
+				ConstraintSetup->DefaultInstance.SetRefFrame(EConstraintFrame::Frame2, SelectedObject.ManipulateTM * StartManParentConTM[i]);
+
+				//Rotation by default only rotates one frame, but translation by default moves both
+				bool bMultiFrame = (InViewportClient->IsAltPressed() && InViewportClient->GetWidgetMode() == FWidget::WM_Rotate) || (!InViewportClient->IsAltPressed() && InViewportClient->GetWidgetMode() == FWidget::WM_Translate);
+
+				if (bMultiFrame)
+				{
+					SharedData->SetConstraintRelTM(&SelectedObject, StartManRelConTM[i]);
+				}
+				else
+				{
+					ConstraintSetup->DefaultInstance.SetRefFrame(EConstraintFrame::Frame1, FTransform(StartManChildConTM[i]));
+				}
+
+				bHandled = true;
+			}
+		}
+	}
+	
+	return bHandled;
+}
+
+void FPhysicsAssetEditorEditMode::Tick(FEditorViewportClient* ViewportClient, float DeltaTime)
+{
+	if (SharedData->bRunningSimulation)
+	{
+		// check if PIE disabled the realtime viewport and quit sim if so
+		if (!ViewportClient->IsRealtime())
+		{
+			SharedData->ToggleSimulation();
+
+			ViewportClient->Invalidate();
+		}
+
+		UWorld* World = SharedData->PreviewScene.Pin()->GetWorld();
+		AWorldSettings* Setting = World->GetWorldSettings();
+		Setting->WorldGravityZ = SharedData->bNoGravitySimulation ? 0.0f : UPhysicsSettings::Get()->DefaultGravityZ*SharedData->EditorOptions->GravScale;
+		Setting->bWorldGravitySet = true;
+
+		// We back up the transforms array now
+		SharedData->EditorSkelComp->AnimationSpaceBases = SharedData->EditorSkelComp->GetComponentSpaceTransforms();
+		SharedData->EditorSkelComp->SetPhysicsBlendWeight(SharedData->EditorOptions->PhysicsBlend);
+		SharedData->EditorSkelComp->bUpdateJointsFromAnimation = SharedData->EditorOptions->bUpdateJointsFromAnimation;
+		SharedData->EditorSkelComp->PhysicsTransformUpdateMode = SharedData->EditorOptions->PhysicsUpdateMode;
+
+		static FPhysicalAnimationData EmptyProfile;
+
+		SharedData->PhysicalAnimationComponent->ApplyPhysicalAnimationProfileBelow(NAME_None, SharedData->PhysicsAsset->CurrentPhysicalAnimationProfileName, /*Include Self=*/true, /*Clear Not Found=*/true);
+	}
+}
+
+void FPhysicsAssetEditorEditMode::Render(const FSceneView* View, FViewport* Viewport, FPrimitiveDrawInterface* PDI)
+{
+	EPhysicsAssetEditorRenderMode MeshViewMode = SharedData->GetCurrentMeshViewMode(SharedData->bRunningSimulation);
+
+	if (MeshViewMode != EPhysicsAssetEditorRenderMode::None)
+	{
+		SharedData->EditorSkelComp->SetVisibility(true);
+
+		if (MeshViewMode == EPhysicsAssetEditorRenderMode::Wireframe)
+		{
+			SharedData->EditorSkelComp->SetForceWireframe(true);
+		}
+		else
+		{
+			SharedData->EditorSkelComp->SetForceWireframe(false);
+		}
+	}
+	else
+	{
+		SharedData->EditorSkelComp->SetVisibility(false);
+	}
+
+	// Draw phat skeletal component.
+	SharedData->EditorSkelComp->DebugDraw(View, PDI);
+}
+
+void FPhysicsAssetEditorEditMode::DrawHUD(FEditorViewportClient* ViewportClient, FViewport* Viewport, const FSceneView* View, FCanvas* Canvas)
+{
+	float W, H;
+	PhysicsAssetEditorFont->GetCharSize(TEXT('L'), W, H);
+
+	const float XOffset = 5.0f;
+	const float YOffset = 48.0f;
+
+	FCanvasTextItem TextItem(FVector2D::ZeroVector, FText::GetEmpty(), PhysicsAssetEditorFont, FLinearColor::White);
+
+	// Write body/constraint count at top.
+	FString StatusString = FText::Format(
+		NSLOCTEXT("UnrealEd", "BodiesConstraints_F", "{0} Bodies  {1} Considered for bounds  {2} Ratio  {3} Constraints"),
+		FText::AsNumber(SharedData->PhysicsAsset->SkeletalBodySetups.Num()),
+		FText::AsNumber(SharedData->PhysicsAsset->BoundsBodies.Num()),
+		FText::AsNumber(static_cast<float>(SharedData->PhysicsAsset->BoundsBodies.Num()) / static_cast<float>(SharedData->PhysicsAsset->SkeletalBodySetups.Num())),
+		FText::AsNumber(SharedData->PhysicsAsset->ConstraintSetup.Num())).ToString();
+
+	TextItem.Text = FText::FromString(StatusString);
+	Canvas->DrawItem(TextItem, XOffset, YOffset);
+
+	TextItem.Text = FText::GetEmpty();
+	if (SharedData->bRunningSimulation)
+	{
+#if PLATFORM_MAC
+		TextItem.Text = LOCTEXT("Sim_Mac", "Command+RightMouse to interact with bodies");
+#else
+		TextItem.Text = LOCTEXT("Sim", "Ctrl+RightMouse to interact with bodies");
+#endif
+	}
+	else if(SharedData->GetSelectedConstraint() != nullptr)
+	{
+		if (ViewportClient->GetWidgetMode() == FWidget::WM_Translate)
+		{
+			TextItem.Text = LOCTEXT("SingleMove", "Hold ALT to move a single reference frame");
+		}
+		else if (ViewportClient->GetWidgetMode() == FWidget::WM_Rotate)
+		{
+			TextItem.Text = LOCTEXT("DoubleRotate", "Hold ALT to rotate both reference frames");
+		}
+	}
+
+	Canvas->DrawItem(TextItem, XOffset, Viewport->GetSizeXY().Y - (3 + H));
+
+	// Draw current physics weight
+	if (SharedData->bRunningSimulation)
+	{
+		FString PhysWeightString = FString::Printf(TEXT("Phys Blend: %3.0f pct"), SharedData->EditorOptions->PhysicsBlend * 100.f);
+		int32 PWLW, PWLH;
+		StringSize(PhysicsAssetEditorFont, PWLW, PWLH, *PhysWeightString);
+		TextItem.Text = FText::FromString(PhysWeightString);
+		Canvas->DrawItem(TextItem, Viewport->GetSizeXY().X - (3 + PWLW + 2 * W), Viewport->GetSizeXY().Y - (3 + H));
+	}
+
+	int32 HalfX = (Viewport->GetSizeXY().X - XOffset) / 2;
+	int32 HalfY = Viewport->GetSizeXY().Y / 2;
+
+	// If showing center-of-mass, and physics is started up..
+	if (SharedData->bShowCOM)
+	{
+		// iterate over each bone
+		for (int32 i = 0; i <SharedData->EditorSkelComp->Bodies.Num(); ++i)
+		{
+			FBodyInstance* BodyInst = SharedData->EditorSkelComp->Bodies[i];
+			check(BodyInst);
+
+			FVector BodyCOMPos = BodyInst->GetCOMPosition();
+			float BodyMass = BodyInst->GetBodyMass();
+
+			FPlane Projection = View->Project(BodyCOMPos);
+			if (Projection.W > 0.f) // This avoids drawing bone names that are behind us.
+			{
+				int32 XPos = HalfX + (HalfX * Projection.X);
+				int32 YPos = HalfY + (HalfY * (Projection.Y * -1));
+
+				FString COMString = FString::Printf(TEXT("%3.3f"), BodyMass);
+				TextItem.Text = FText::FromString(COMString);
+				TextItem.SetColor(SharedData->COMRenderColor);
+				Canvas->DrawItem(TextItem, XPos, YPos);
+			}
+		}
+	}
+}
+
+bool FPhysicsAssetEditorEditMode::AllowWidgetMove()
+{
+	return true;
+}
+
+bool FPhysicsAssetEditorEditMode::ShouldDrawWidget() const
+{
+	return !SharedData->bRunningSimulation && (SharedData->GetSelectedBody() || SharedData->GetSelectedConstraint());
+}
+
+bool FPhysicsAssetEditorEditMode::UsesTransformWidget() const
+{
+	return ShouldDrawWidget();
+}
+
+bool FPhysicsAssetEditorEditMode::UsesTransformWidget(FWidget::EWidgetMode CheckMode) const
+{
+	if (SharedData->GetSelectedConstraint() && CheckMode == FWidget::WM_Scale)
+	{
+		return false;
+	}
+
+	return ShouldDrawWidget() && (CheckMode == FWidget::WM_Scale || CheckMode == FWidget::WM_Translate || CheckMode == FWidget::WM_Rotate);
+}
+
+bool FPhysicsAssetEditorEditMode::GetCustomDrawingCoordinateSystem(FMatrix& InMatrix, void* InData)
+{
+	// Don't draw widget if nothing selected.
+	if (SharedData->GetSelectedBody())
+	{
+		int32 BoneIndex = SharedData->EditorSkelComp->GetBoneIndex(SharedData->PhysicsAsset->SkeletalBodySetups[SharedData->GetSelectedBody()->Index]->BoneName);
+
+		FTransform BoneTM = SharedData->EditorSkelComp->GetBoneTransform(BoneIndex);
+		BoneTM.RemoveScaling();
+
+		InMatrix = SharedData->EditorSkelComp->GetPrimitiveTransform(BoneTM, SharedData->GetSelectedBody()->Index, SharedData->GetSelectedBody()->PrimitiveType, SharedData->GetSelectedBody()->PrimitiveIndex, 1.f).ToMatrixNoScale().RemoveTranslation();
+		return true;
+	}
+	else if (SharedData->GetSelectedConstraint())
+	{
+		InMatrix = SharedData->GetConstraintMatrix(SharedData->GetSelectedConstraint()->Index, EConstraintFrame::Frame2, 1.f).ToMatrixNoScale().RemoveTranslation();
+		return true;
+	}
+
+	return false;
+}
+
+bool FPhysicsAssetEditorEditMode::GetCustomInputCoordinateSystem(FMatrix& InMatrix, void* InData)
+{
+	return GetCustomDrawingCoordinateSystem(InMatrix, InData);
+}
+
+FVector FPhysicsAssetEditorEditMode::GetWidgetLocation() const
+{
+	// Don't draw widget if nothing selected.
+	if (SharedData->GetSelectedBody())
+	{
+		int32 BoneIndex = SharedData->EditorSkelComp->GetBoneIndex(SharedData->PhysicsAsset->SkeletalBodySetups[SharedData->GetSelectedBody()->Index]->BoneName);
+
+		FTransform BoneTM = SharedData->EditorSkelComp->GetBoneTransform(BoneIndex);
+		const float Scale = BoneTM.GetScale3D().GetAbsMax();
+		BoneTM.RemoveScaling();
+
+		return SharedData->EditorSkelComp->GetPrimitiveTransform(BoneTM, SharedData->GetSelectedBody()->Index, SharedData->GetSelectedBody()->PrimitiveType, SharedData->GetSelectedBody()->PrimitiveIndex, Scale).GetTranslation();
+	}
+	else if (SharedData->GetSelectedConstraint())
+	{
+		return SharedData->GetConstraintMatrix(SharedData->GetSelectedConstraint()->Index, EConstraintFrame::Frame2, 1.f).GetTranslation();
+	}
+
+	return FVector::ZeroVector;
+}
+
+bool FPhysicsAssetEditorEditMode::HandleClick(FEditorViewportClient* InViewportClient, HHitProxy *HitProxy, const FViewportClick &Click)
+{
+	if (!SharedData->bRunningSimulation)
+	{
+		if (Click.GetKey() == EKeys::LeftMouseButton)
+		{
+			if (HitProxy && HitProxy->IsA(HPhysicsAssetEditorEdBoneProxy::StaticGetType()))
+			{
+				HPhysicsAssetEditorEdBoneProxy* BoneProxy = (HPhysicsAssetEditorEdBoneProxy*)HitProxy;
+
+				SharedData->HitBone(BoneProxy->BodyIndex, BoneProxy->PrimType, BoneProxy->PrimIndex, InViewportClient->IsCtrlPressed() || InViewportClient->IsShiftPressed());
+				return true;
+			}
+			else if (HitProxy && HitProxy->IsA(HPhysicsAssetEditorEdConstraintProxy::StaticGetType()))
+			{
+				HPhysicsAssetEditorEdConstraintProxy* ConstraintProxy = (HPhysicsAssetEditorEdConstraintProxy*)HitProxy;
+
+				SharedData->HitConstraint(ConstraintProxy->ConstraintIndex, InViewportClient->IsCtrlPressed() || InViewportClient->IsShiftPressed());
+				return true;
+			}
+			else
+			{
+				HitNothing(InViewportClient);
+			}
+		}
+		else if (Click.GetKey() == EKeys::RightMouseButton)
+		{
+			if (HitProxy && HitProxy->IsA(HPhysicsAssetEditorEdBoneProxy::StaticGetType()))
+			{
+				HPhysicsAssetEditorEdBoneProxy* BoneProxy = (HPhysicsAssetEditorEdBoneProxy*)HitProxy;
+
+				// Select body under cursor if not already selected	(if ctrl is held down we only add, not remove)
+				FPhysicsAssetEditorSharedData::FSelection Selection(BoneProxy->BodyIndex, BoneProxy->PrimType, BoneProxy->PrimIndex);
+				if (!SharedData->IsBodySelected(Selection))
+				{
+					SharedData->SetSelectedBody(Selection, true);
+				}
+
+				// Pop up menu, if we have a body selected.
+				if (SharedData->GetSelectedBody())
+				{
+					OpenBodyMenu(InViewportClient);
+				}
+
+				return true;
+			}
+			else if (HitProxy && HitProxy->IsA(HPhysicsAssetEditorEdConstraintProxy::StaticGetType()))
+			{
+				HPhysicsAssetEditorEdConstraintProxy* ConstraintProxy = (HPhysicsAssetEditorEdConstraintProxy*)HitProxy;
+
+				// Select constraint under cursor if not already selected (if ctrl is held down we only add, not remove)
+				if (!SharedData->IsConstraintSelected(ConstraintProxy->ConstraintIndex))
+				{
+					SharedData->SetSelectedConstraint(ConstraintProxy->ConstraintIndex, InViewportClient->IsCtrlPressed());
+				}
+
+				// Pop up menu, if we have a constraint selected.
+				if (SharedData->GetSelectedConstraint())
+				{
+					OpenConstraintMenu(InViewportClient);
+				}
+
+				return true;
+			}
+			else
+			{
+				OpenSelectionMenu(InViewportClient);
+				return true;
+			}
+		}
+	}
+
+	return false;
+}
+
+/** Helper function to open a viewport context menu */
+static void OpenContextMenu(const TSharedRef<FPhysicsAssetEditor>& PhysicsAssetEditor, FEditorViewportClient* InViewportClient, TFunctionRef<void(FMenuBuilder&)> InBuildMenu)
+{
+	FMenuBuilder MenuBuilder(true, PhysicsAssetEditor->GetToolkitCommands());
+
+	InBuildMenu(MenuBuilder);
+
+	TSharedPtr<SWidget> MenuWidget = MenuBuilder.MakeWidget();
+	TSharedPtr<SWidget> ParentWidget = InViewportClient->GetEditorViewportWidget();
+
+	if (MenuWidget.IsValid() && ParentWidget.IsValid())
+	{
+		const FVector2D MouseCursorLocation = FSlateApplication::Get().GetCursorPos();
+
+		FSlateApplication::Get().PushMenu(
+			ParentWidget.ToSharedRef(),
+			FWidgetPath(),
+			MenuWidget.ToSharedRef(),
+			MouseCursorLocation,
+			FPopupTransitionEffect(FPopupTransitionEffect::ContextMenu)
+		);
+	}
+}
+
+void FPhysicsAssetEditorEditMode::OpenBodyMenu(FEditorViewportClient* InViewportClient)
+{
+	OpenContextMenu(PhysicsAssetEditorPtr.Pin().ToSharedRef(), InViewportClient, 
+		[this](FMenuBuilder& InMenuBuilder)
+		{
+			PhysicsAssetEditorPtr.Pin()->BuildMenuWidgetBody(InMenuBuilder);
+			PhysicsAssetEditorPtr.Pin()->BuildMenuWidgetSelection(InMenuBuilder);
+		});
+}
+
+void FPhysicsAssetEditorEditMode::OpenConstraintMenu(FEditorViewportClient* InViewportClient)
+{
+	OpenContextMenu(PhysicsAssetEditorPtr.Pin().ToSharedRef(), InViewportClient, 
+		[this](FMenuBuilder& InMenuBuilder)
+		{
+			PhysicsAssetEditorPtr.Pin()->BuildMenuWidgetConstraint(InMenuBuilder);
+			PhysicsAssetEditorPtr.Pin()->BuildMenuWidgetSelection(InMenuBuilder);
+		});
+}
+
+void FPhysicsAssetEditorEditMode::OpenSelectionMenu(FEditorViewportClient* InViewportClient)
+{
+	OpenContextMenu(PhysicsAssetEditorPtr.Pin().ToSharedRef(), InViewportClient, 
+		[this](FMenuBuilder& InMenuBuilder)
+		{
+			PhysicsAssetEditorPtr.Pin()->BuildMenuWidgetSelection(InMenuBuilder);
+		});
+}
+
+bool FPhysicsAssetEditorEditMode::SimMousePress(FEditorViewportClient* InViewportClient, FKey Key)
+{
+	FViewport* Viewport = InViewportClient->Viewport;
+
+	bool bCtrlDown = Viewport->KeyState(EKeys::LeftControl) || Viewport->KeyState(EKeys::RightControl);
+	bool bShiftDown = Viewport->KeyState(EKeys::LeftShift) || Viewport->KeyState(EKeys::RightShift);
+
+	FSceneViewFamilyContext ViewFamily(FSceneViewFamily::ConstructionValues(Viewport, InViewportClient->GetScene(), InViewportClient->EngineShowFlags));
+	FSceneView* View = InViewportClient->CalcSceneView(&ViewFamily);
+
+	const FViewportClick Click(View, InViewportClient, EKeys::Invalid, IE_Released, Viewport->GetMouseX(), Viewport->GetMouseY());
+#if DEBUG_CLICK_VIEWPORT	
+	SharedData->LastClickOrigin = Click.GetOrigin();
+	SharedData->LastClickDirection = Click.GetDirection();
+#endif
+	SharedData->LastClickPos = Click.GetClickPos();
+	FHitResult Result(1.f);
+	bool bHit = SharedData->EditorSkelComp->LineTraceComponent(Result, Click.GetOrigin() - Click.GetDirection() * SimGrabCheckDistance, Click.GetOrigin() + Click.GetDirection() * SimGrabCheckDistance, FCollisionQueryParams(NAME_None, true));
+
+	if (bHit && (bCtrlDown || bShiftDown))
+	{
+		check(Result.Item != INDEX_NONE);
+		FName BoneName = SharedData->PhysicsAsset->SkeletalBodySetups[Result.Item]->BoneName;
+
+		//UE_LOG(LogPhysics, Warning, TEXT("Hit Bone Name (%s)"), *BoneName.ToString());
+
+		// Right mouse is for dragging things around
+		if (Key == EKeys::RightMouseButton)
+		{
+			SharedData->bManipulating = true;
+			DragX = 0.0f;
+			DragY = 0.0f;
+			SimGrabPush = 0.0f;
+
+			// Update mouse force properties from sim options.
+			SharedData->MouseHandle->LinearDamping = SharedData->EditorOptions->HandleLinearDamping;
+			SharedData->MouseHandle->LinearStiffness = SharedData->EditorOptions->HandleLinearStiffness;
+			SharedData->MouseHandle->AngularDamping = SharedData->EditorOptions->HandleAngularDamping;
+			SharedData->MouseHandle->AngularStiffness = SharedData->EditorOptions->HandleAngularStiffness;
+			SharedData->MouseHandle->InterpolationSpeed = SharedData->EditorOptions->InterpolationSpeed;
+
+			// Create handle to object.
+			SharedData->MouseHandle->GrabComponentAtLocationWithRotation(SharedData->EditorSkelComp, BoneName, Result.Location, FRotator::ZeroRotator);
+
+			FMatrix	InvViewMatrix = View->ViewMatrices.GetInvViewMatrix();
+
+			SimGrabMinPush = SimMinHoldDistance - (Result.Time * SimGrabCheckDistance);
+
+			SimGrabLocation = Result.Location;
+			SimGrabX = InvViewMatrix.GetUnitAxis(EAxis::X);
+			SimGrabY = InvViewMatrix.GetUnitAxis(EAxis::Y);
+			SimGrabZ = InvViewMatrix.GetUnitAxis(EAxis::Z);
+		}
+		// Left mouse is for poking things
+		else if (Key == EKeys::LeftMouseButton)
+		{
+			SharedData->EditorSkelComp->AddImpulseAtLocation(Click.GetDirection() * SharedData->EditorOptions->PokeStrength, Result.Location, BoneName);
+		}
+
+		return true;
+	}
+
+	return false;
+}
+
+void FPhysicsAssetEditorEditMode::SimMouseMove(FEditorViewportClient* InViewportClient, float DeltaX, float DeltaY)
+{
+	DragX = InViewportClient->Viewport->GetMouseX() - SharedData->LastClickPos.X;
+	DragY = InViewportClient->Viewport->GetMouseY() - SharedData->LastClickPos.Y;
+
+	if (!SharedData->MouseHandle->GrabbedComponent)
+	{
+		return;
+	}
+
+	//We need to convert Pixel Delta into Screen position (deal with different viewport sizes)
+	FSceneViewFamilyContext ViewFamily(FSceneViewFamily::ConstructionValues(InViewportClient->Viewport, SharedData->PreviewScene.Pin()->GetScene(), InViewportClient->EngineShowFlags));
+	FSceneView* View = InViewportClient->CalcSceneView(&ViewFamily);
+	FVector4 ScreenOldPos = View->PixelToScreen(SharedData->LastClickPos.X, SharedData->LastClickPos.Y, 1.f);
+	FVector4 ScreenNewPos = View->PixelToScreen(DragX + SharedData->LastClickPos.X, DragY + SharedData->LastClickPos.Y, 1.f);
+	FVector4 ScreenDelta = ScreenNewPos - ScreenOldPos;
+	FVector4 ProjectedDelta = View->ScreenToWorld(ScreenDelta);
+	FVector4 WorldDelta;
+
+	//Now we project new ScreenPos to xy-plane of SimGrabLocation
+	FVector LocalOffset = View->ViewMatrices.GetViewMatrix().TransformPosition(SimGrabLocation + SimGrabZ * SimGrabPush);
+	float ZDistance = InViewportClient->GetViewportType() == ELevelViewportType::LVT_Perspective ? FMath::Abs(LocalOffset.Z) : 1.f;	//in the ortho case we don't need to do any fixup because there is no perspective
+	WorldDelta = ProjectedDelta * ZDistance;
+
+	//Now we convert back into WorldPos
+	FVector WorldPos = SimGrabLocation + WorldDelta + SimGrabZ * SimGrabPush;
+	FVector NewLocation = WorldPos;
+	float QuickRadius = 5 - SimGrabPush / SimHoldDistanceChangeDelta;
+	QuickRadius = QuickRadius < 2 ? 2 : QuickRadius;
+
+	DrawDebugPoint(GetWorld(), NewLocation, QuickRadius, FColorList::Red, false, 0.3);
+
+	SharedData->MouseHandle->SetTargetLocation(NewLocation);
+	SharedData->MouseHandle->GrabbedComponent->WakeRigidBody(SharedData->MouseHandle->GrabbedBoneName);
+}
+
+bool FPhysicsAssetEditorEditMode::SimMouseRelease()
+{
+	SharedData->bManipulating = false;
+
+	if (!SharedData->MouseHandle->GrabbedComponent)
+	{
+		return false;
+	}
+
+	SharedData->MouseHandle->GrabbedComponent->WakeRigidBody(SharedData->MouseHandle->GrabbedBoneName);
+	SharedData->MouseHandle->ReleaseComponent();
+
+	return true;
+}
+
+bool FPhysicsAssetEditorEditMode::SimMouseWheelUp(FEditorViewportClient* InViewportClient)
+{
+	if (!SharedData->MouseHandle->GrabbedComponent)
+	{
+		return false;
+	}
+
+	SimGrabPush += SimHoldDistanceChangeDelta;
+
+	SimMouseMove(InViewportClient, 0.0f, 0.0f);
+
+	return true;
+}
+
+bool FPhysicsAssetEditorEditMode::SimMouseWheelDown(FEditorViewportClient* InViewportClient)
+{
+	if (!SharedData->MouseHandle->GrabbedComponent)
+	{
+		return false;
+	}
+
+	SimGrabPush -= SimHoldDistanceChangeDelta;
+	SimGrabPush = FMath::Max(SimGrabMinPush, SimGrabPush);
+
+	SimMouseMove(InViewportClient, 0.0f, 0.0f);
+
+	return true;
+}
+
+void FPhysicsAssetEditorEditMode::ModifyPrimitiveSize(int32 BodyIndex, EAggCollisionShape::Type PrimType, int32 PrimIndex, FVector DeltaSize)
+{
+	check(SharedData->GetSelectedBody());
+
+	FKAggregateGeom* AggGeom = &SharedData->PhysicsAsset->SkeletalBodySetups[BodyIndex]->AggGeom;
+
+	if (PrimType == EAggCollisionShape::Sphere)
+	{
+		check(AggGeom->SphereElems.IsValidIndex(PrimIndex));
+		AggGeom->SphereElems[PrimIndex].ScaleElem(DeltaSize, MinPrimSize);
+	}
+	else if (PrimType == EAggCollisionShape::Box)
+	{
+		check(AggGeom->BoxElems.IsValidIndex(PrimIndex));
+		AggGeom->BoxElems[PrimIndex].ScaleElem(DeltaSize, MinPrimSize);
+	}
+	else if (PrimType == EAggCollisionShape::Sphyl)
+	{
+		check(AggGeom->SphylElems.IsValidIndex(PrimIndex));
+		AggGeom->SphylElems[PrimIndex].ScaleElem(DeltaSize, MinPrimSize);
+	}
+	else if (PrimType == EAggCollisionShape::Convex)
+	{
+		check(AggGeom->ConvexElems.IsValidIndex(PrimIndex));
+
+		FVector ModifiedSize = DeltaSize;
+		if (GEditor->UsePercentageBasedScaling())
+		{
+			ModifiedSize = DeltaSize * ((GEditor->GetScaleGridSize() / 100.0f) / GEditor->GetGridSize());
+		}
+
+		AggGeom->ConvexElems[PrimIndex].ScaleElem(ModifiedSize, MinPrimSize);
+	}
+}
+
+void FPhysicsAssetEditorEditMode::HitNothing(FEditorViewportClient* InViewportClient)
+{
+	if (InViewportClient->IsCtrlPressed() == false)	//we only want to deselect if Ctrl is not used
+	{
+		SharedData->ClearSelectedBody();
+		SharedData->ClearSelectedConstraints();
+	}
+
+	InViewportClient->Invalidate();
+	PhysicsAssetEditorPtr.Pin()->RefreshHierachyTree();
+}
+
+#undef LOCTEXT_NAMESPACE
diff --git a/Engine/Source/Editor/PhysicsAssetEditor/Private/PhysicsAssetEditorMode.cpp b/Engine/Source/Editor/PhysicsAssetEditor/Private/PhysicsAssetEditorMode.cpp
new file mode 100644
--- /dev/null
+++ b/Engine/Source/Editor/PhysicsAssetEditor/Private/PhysicsAssetEditorMode.cpp
@@ -0,0 +1,308 @@
+// Copyright 1998-2016 Epic Games, Inc. All Rights Reserved.
+
+#include "PhysicsAssetEditorMode.h"
+#include "PhysicsAssetEditor.h"
+#include "ISkeletonTree.h"
+#include "IPersonaPreviewScene.h"
+#include "PersonaModule.h"
+#include "ISkeletonEditorModule.h"
+#include "PhysicsAssetGraphSummoner.h"
+#include "PhysicsEngine/PhysicsAsset.h"
+#include "PhysicsAssetEditorActions.h"
+#include "SEditorViewportToolBarMenu.h"
+#include "PhysicsAssetEditorProfilesSummoner.h"
+#include "PropertyEditorModule.h"
+#include "Preferences/PhysicsAssetEditorOptions.h"
+#include "SSpinBox.h"
+#include "ModuleManager.h"
+#include "MultiBoxBuilder.h"
+#include "PhysicsAssetEditorToolsSummoner.h"
+
+#define LOCTEXT_NAMESPACE "PhysicsAssetEditorMode"
+
+static const FName PhysicsAssetEditorPreviewViewportName("Viewport");
+static const FName PhysicsAssetEditorPropertiesName("DetailsTab");
+static const FName PhysicsAssetEditorHierarchyName("SkeletonTreeView");
+static const FName PhysicsAssetEditorGraphName("PhysicsAssetGraphView");
+static const FName PhysicsAssetEditorProfilesName("PhysicsAssetProfilesView");
+static const FName PhysicsAssetEditorToolsName("PhysicsAssetTools");
+static const FName PhysicsAssetEditorAdvancedPreviewName(TEXT("AdvancedPreviewTab"));
+
+FPhysicsAssetEditorMode::FPhysicsAssetEditorMode(TSharedRef<FWorkflowCentricApplication> InHostingApp, TSharedRef<ISkeletonTree> InSkeletonTree, TSharedRef<IPersonaPreviewScene> InPreviewScene)
+	: FApplicationMode(PhysicsAssetEditorModes::PhysicsAssetEditorMode)
+{
+	PhysicsAssetEditorPtr = StaticCastSharedRef<FPhysicsAssetEditor>(InHostingApp);
+	TSharedRef<FPhysicsAssetEditor> PhysicsAssetEditor = StaticCastSharedRef<FPhysicsAssetEditor>(InHostingApp);
+
+	ISkeletonEditorModule& SkeletonEditorModule = FModuleManager::LoadModuleChecked<ISkeletonEditorModule>("SkeletonEditor");
+	TabFactories.RegisterFactory(SkeletonEditorModule.CreateSkeletonTreeTabFactory(InHostingApp, InSkeletonTree));
+
+	FPersonaModule& PersonaModule = FModuleManager::LoadModuleChecked<FPersonaModule>("Persona");
+	TabFactories.RegisterFactory(PersonaModule.CreateDetailsTabFactory(InHostingApp, FOnDetailsCreated::CreateSP(&PhysicsAssetEditor.Get(), &FPhysicsAssetEditor::HandleDetailsCreated)));
+
+	auto ExtendShowMenu = [this](FMenuBuilder& InMenuBuilder)
+	{
+		const FPhysicsAssetEditorCommands& Actions = FPhysicsAssetEditorCommands::Get();
+
+		InMenuBuilder.PushCommandList(PhysicsAssetEditorPtr.Pin()->GetToolkitCommands());
+
+		InMenuBuilder.BeginSection("PhysicsAssetShowCommands", LOCTEXT("PhysicsShowCommands", "Physics Rendering"));
+		{
+			InMenuBuilder.AddMenuEntry(Actions.ToggleMassProperties);
+
+			// Mesh, collision and constraint rendering modes
+			struct Local
+			{
+				static void BuildMeshRenderModeMenu(FMenuBuilder& InSubMenuBuilder)
+				{
+					const FPhysicsAssetEditorCommands& Commands = FPhysicsAssetEditorCommands::Get();
+
+					InSubMenuBuilder.BeginSection("PhysicsAssetEditorRenderingMode", LOCTEXT("MeshRenderModeHeader", "Mesh Drawing (Edit)"));
+					{
+						InSubMenuBuilder.AddMenuEntry(Commands.MeshRenderingMode_Solid);
+						InSubMenuBuilder.AddMenuEntry(Commands.MeshRenderingMode_Wireframe);
+						InSubMenuBuilder.AddMenuEntry(Commands.MeshRenderingMode_None);
+					}
+					InSubMenuBuilder.EndSection();
+
+					InSubMenuBuilder.BeginSection("PhysicsAssetEditorRenderingModeSim", LOCTEXT("MeshRenderModeSimHeader", "Mesh Drawing (Simulation)"));
+					{
+						InSubMenuBuilder.AddMenuEntry(Commands.MeshRenderingMode_Simulation_Solid);
+						InSubMenuBuilder.AddMenuEntry(Commands.MeshRenderingMode_Simulation_Wireframe);
+						InSubMenuBuilder.AddMenuEntry(Commands.MeshRenderingMode_Simulation_None);
+					}
+					InSubMenuBuilder.EndSection();
+				}
+
+				static void BuildCollisionRenderModeMenu(FMenuBuilder& InSubMenuBuilder, TWeakPtr<FPhysicsAssetEditor> PhysicsAssetEditorPtr)
+				{
+					const FPhysicsAssetEditorCommands& Commands = FPhysicsAssetEditorCommands::Get();
+
+					InSubMenuBuilder.BeginSection("PhysicsAssetEditorCollisionRenderSettings", LOCTEXT("CollisionRenderSettingsHeader", "Body Drawing"));
+					{
+						InSubMenuBuilder.AddMenuEntry(Commands.RenderOnlySelectedSolid);
+
+						UPhysicsAssetEditorOptions* Options = PhysicsAssetEditorPtr.Pin()->GetSharedData()->EditorOptions;
+
+						TSharedPtr<SWidget> CollisionOpacityWidget = 
+							SNew(SBox)
+							.HAlign(HAlign_Right)
+							[
+								SNew(SBox)
+								.Padding(FMargin(4.0f, 0.0f, 0.0f, 0.0f))
+								.WidthOverride(100.0f)
+								[
+									SNew(SSpinBox<float>)
+									.Font(FEditorStyle::GetFontStyle(TEXT("MenuItem.Font")))
+									.MinValue(0.0f)
+									.MaxValue(1.0f)
+									.Value_Lambda([Options]() { return Options->CollisionOpacity; })
+									.OnValueChanged_Lambda([Options](float InValue) { Options->CollisionOpacity = InValue; })
+									.OnValueCommitted_Lambda([Options](float InValue, ETextCommit::Type InCommitType) 
+									{ 
+										Options->CollisionOpacity = InValue; 
+										Options->SaveConfig(); 
+									})
+								]
+							];
+
+						InSubMenuBuilder.AddWidget(CollisionOpacityWidget.ToSharedRef(), LOCTEXT("CollisionOpacityLabel", "Collision Opacity"));
+					}
+					InSubMenuBuilder.EndSection();
+
+					InSubMenuBuilder.BeginSection("PhysicsAssetEditorCollisionMode", LOCTEXT("CollisionRenderModeHeader", "Body Drawing (Edit)"));
+					{
+						InSubMenuBuilder.AddMenuEntry(Commands.CollisionRenderingMode_Solid);
+						InSubMenuBuilder.AddMenuEntry(Commands.CollisionRenderingMode_Wireframe);
+						InSubMenuBuilder.AddMenuEntry(Commands.CollisionRenderingMode_None);
+					}
+					InSubMenuBuilder.EndSection();
+
+					InSubMenuBuilder.BeginSection("PhysicsAssetEditorCollisionModeSim", LOCTEXT("CollisionRenderModeSimHeader", "Body Drawing (Simulation)"));
+					{
+						InSubMenuBuilder.AddMenuEntry(Commands.CollisionRenderingMode_Simulation_Solid);
+						InSubMenuBuilder.AddMenuEntry(Commands.CollisionRenderingMode_Simulation_Wireframe);
+						InSubMenuBuilder.AddMenuEntry(Commands.CollisionRenderingMode_Simulation_None);
+					}
+					InSubMenuBuilder.EndSection();
+				}
+
+				static void BuildConstraintRenderModeMenu(FMenuBuilder& InSubMenuBuilder, TWeakPtr<FPhysicsAssetEditor> PhysicsAssetEditorPtr)
+				{
+					const FPhysicsAssetEditorCommands& Commands = FPhysicsAssetEditorCommands::Get();
+
+					InSubMenuBuilder.BeginSection("PhysicsAssetEditorConstraints", LOCTEXT("ConstraintHeader", "Constraints"));
+					{
+						InSubMenuBuilder.AddMenuEntry(Commands.DrawConstraintsAsPoints);
+
+						UPhysicsAssetEditorOptions* Options = PhysicsAssetEditorPtr.Pin()->GetSharedData()->EditorOptions;
+
+						TSharedPtr<SWidget> ConstraintScaleWidget = 
+							SNew(SBox)
+							.HAlign(HAlign_Right)
+							[
+								SNew(SBox)
+								.Padding(FMargin(4.0f, 0.0f, 0.0f, 0.0f))
+								.WidthOverride(100.0f)
+								[
+									SNew(SSpinBox<float>)
+									.Font(FEditorStyle::GetFontStyle(TEXT("MenuItem.Font")))
+									.MinValue(0.0f)
+									.MaxValue(4.0f)
+									.Value_Lambda([Options]() { return Options->ConstraintDrawSize; })
+									.OnValueChanged_Lambda([Options](float InValue) { Options->ConstraintDrawSize = InValue; })
+									.OnValueCommitted_Lambda([Options](float InValue, ETextCommit::Type InCommitType) { Options->ConstraintDrawSize = InValue; Options->SaveConfig(); })
+								]
+							];
+
+						InSubMenuBuilder.AddWidget(ConstraintScaleWidget.ToSharedRef(), LOCTEXT("ConstraintScaleLabel", "Constraint Scale"));
+					}
+					InSubMenuBuilder.EndSection();
+
+					InSubMenuBuilder.BeginSection("PhysicsAssetEditorConstraintMode", LOCTEXT("ConstraintRenderModeHeader", "Constraint Drawing (Edit)"));
+					{
+						InSubMenuBuilder.AddMenuEntry(Commands.ConstraintRenderingMode_None);
+						InSubMenuBuilder.AddMenuEntry(Commands.ConstraintRenderingMode_AllPositions);
+						InSubMenuBuilder.AddMenuEntry(Commands.ConstraintRenderingMode_AllLimits);
+					}
+					InSubMenuBuilder.EndSection();
+
+					InSubMenuBuilder.BeginSection("PhysicsAssetEditorConstraintModeSim", LOCTEXT("ConstraintRenderModeSimHeader", "Constraint Drawing (Simulation)"));
+					{
+						InSubMenuBuilder.AddMenuEntry(Commands.ConstraintRenderingMode_Simulation_None);
+						InSubMenuBuilder.AddMenuEntry(Commands.ConstraintRenderingMode_Simulation_AllPositions);
+						InSubMenuBuilder.AddMenuEntry(Commands.ConstraintRenderingMode_Simulation_AllLimits);
+					}
+					InSubMenuBuilder.EndSection();
+				}
+			};
+
+			InMenuBuilder.AddSubMenu(LOCTEXT("MeshRenderModeSubMenu", "Mesh"), FText::GetEmpty(), FNewMenuDelegate::CreateStatic(&Local::BuildMeshRenderModeMenu));
+			InMenuBuilder.AddSubMenu(LOCTEXT("CollisionRenderModeSubMenu", "Bodies"), FText::GetEmpty(), FNewMenuDelegate::CreateStatic(&Local::BuildCollisionRenderModeMenu, PhysicsAssetEditorPtr));
+			InMenuBuilder.AddSubMenu(LOCTEXT("ConstraintRenderModeSubMenu", "Constraints"), FText::GetEmpty(), FNewMenuDelegate::CreateStatic(&Local::BuildConstraintRenderModeMenu, PhysicsAssetEditorPtr));
+		}
+		InMenuBuilder.EndSection();
+
+		InMenuBuilder.PopCommandList();
+	};
+
+	auto ExtendMenuBar = [this](FMenuBarBuilder& InMenuBarBuilder)
+	{
+		InMenuBarBuilder.AddPullDownMenu(
+			LOCTEXT("PhysicsMenu", "Physics"),
+			LOCTEXT("PhysicsMenuTooltip", "Adjust physics settings"),
+			FSlateIcon(),
+			FNewMenuDelegate::CreateLambda([this](FMenuBuilder& InSubMenuBuilder)
+			{
+				FPropertyEditorModule& PropertyEditorModule = FModuleManager::GetModuleChecked<FPropertyEditorModule>("PropertyEditor");
+
+				FDetailsViewArgs DetailsViewArgs;
+				DetailsViewArgs.bAllowSearch = false;
+				DetailsViewArgs.NameAreaSettings = FDetailsViewArgs::HideNameArea;
+				TSharedPtr<IDetailsView> DetailsView = PropertyEditorModule.CreateDetailView(DetailsViewArgs);
+				DetailsView->SetObject(PhysicsAssetEditorPtr.Pin()->GetSharedData()->EditorOptions);
+				DetailsView->OnFinishedChangingProperties().AddLambda([this](const FPropertyChangedEvent& InEvent){ PhysicsAssetEditorPtr.Pin()->GetSharedData()->EditorOptions->SaveConfig(); });
+
+				InSubMenuBuilder.AddWidget(DetailsView.ToSharedRef(), FText(), true);
+			})
+		);
+	};
+
+	TSharedRef<FExtender> ViewportExtender = MakeShared<FExtender>();
+	ViewportExtender->AddMenuExtension("AnimViewportGeneralShowFlags", EExtensionHook::After, PhysicsAssetEditor->GetToolkitCommands(), FMenuExtensionDelegate::CreateLambda(ExtendShowMenu));
+	ViewportExtender->AddMenuBarExtension("AnimViewportShowMenu", EExtensionHook::After, PhysicsAssetEditor->GetToolkitCommands(), FMenuBarExtensionDelegate::CreateLambda(ExtendMenuBar));
+
+	FPersonaViewportArgs ViewportArgs(InSkeletonTree, InPreviewScene, PhysicsAssetEditor->OnPostUndo);
+	ViewportArgs.bAlwaysShowTransformToolbar = true;
+	ViewportArgs.bShowStats = false;
+	ViewportArgs.bShowTurnTable = false;
+	ViewportArgs.Extenders = ViewportExtender;
+
+	TabFactories.RegisterFactory(PersonaModule.CreatePersonaViewportTabFactory(InHostingApp, ViewportArgs));
+
+	TabFactories.RegisterFactory(PersonaModule.CreateAdvancedPreviewSceneTabFactory(InHostingApp, InPreviewScene));
+
+	TabFactories.RegisterFactory(MakeShared<FPhysicsAssetGraphSummoner>(InHostingApp, CastChecked<UPhysicsAsset>((*PhysicsAssetEditor->GetObjectsCurrentlyBeingEdited())[0]), InSkeletonTree->GetEditableSkeleton(), FOnPhysicsAssetGraphCreated::CreateSP(&PhysicsAssetEditor.Get(), &FPhysicsAssetEditor::HandlePhysicsAssetGraphCreated), FOnGraphObjectsSelected::CreateSP(&PhysicsAssetEditor.Get(), &FPhysicsAssetEditor::HandleGraphObjectsSelected)));
+
+	TabFactories.RegisterFactory(MakeShared<FPhysicsAssetEditorProfilesSummoner>(InHostingApp, CastChecked<UPhysicsAsset>((*PhysicsAssetEditor->GetObjectsCurrentlyBeingEdited())[0])));
+
+	TabFactories.RegisterFactory(MakeShared<FPhysicsAssetEditorToolsSummoner>(InHostingApp));
+
+	TabLayout = FTabManager::NewLayout("Standalone_PhysicsAssetEditor_Layout_v5.2")
+		->AddArea
+		(
+			FTabManager::NewPrimaryArea()
+			->SetOrientation(Orient_Vertical)
+			->Split
+			(
+				FTabManager::NewStack()
+				->SetSizeCoefficient(0.1f)
+				->SetHideTabWell(true)
+				->AddTab(InHostingApp->GetToolbarTabId(), ETabState::OpenedTab)
+			)
+			->Split
+			(
+				FTabManager::NewSplitter()
+				->SetSizeCoefficient(0.9f)
+				->SetOrientation(Orient_Horizontal)
+				->Split
+				(
+				    FTabManager::NewSplitter()
+				    ->SetSizeCoefficient(0.2f)
+				    ->SetOrientation(Orient_Vertical)
+				    ->Split
+				    (
+					    FTabManager::NewStack()
+					    ->SetSizeCoefficient(0.6f)
+					    ->AddTab(PhysicsAssetEditorHierarchyName, ETabState::OpenedTab)
+					)
+					->Split
+					(
+					    FTabManager::NewStack()
+					    ->SetSizeCoefficient(0.4f)
+					    ->AddTab(PhysicsAssetEditorGraphName, ETabState::OpenedTab)
+					)
+				)
+				->Split
+				(
+					FTabManager::NewStack()
+					->SetSizeCoefficient(0.6f)
+					->SetHideTabWell(true)
+					->AddTab(PhysicsAssetEditorPreviewViewportName, ETabState::OpenedTab)
+				)
+				->Split
+				(
+					FTabManager::NewSplitter()
+					->SetSizeCoefficient(0.2f)
+					->SetOrientation(Orient_Vertical)
+					->Split
+					(
+						FTabManager::NewStack()
+						->SetSizeCoefficient(0.6f)
+						->AddTab(PhysicsAssetEditorPropertiesName, ETabState::OpenedTab)
+						->AddTab(PhysicsAssetEditorAdvancedPreviewName, ETabState::OpenedTab)
+						->SetForegroundTab(PhysicsAssetEditorPropertiesName)
+					)
+					->Split
+					(
+						FTabManager::NewStack()
+						->SetSizeCoefficient(0.4f)
+						->AddTab(PhysicsAssetEditorToolsName, ETabState::OpenedTab)
+						->AddTab(PhysicsAssetEditorProfilesName, ETabState::OpenedTab)
+						->SetForegroundTab(PhysicsAssetEditorToolsName)
+					)
+				)
+			)
+		);
+}
+
+void FPhysicsAssetEditorMode::RegisterTabFactories(TSharedPtr<FTabManager> InTabManager)
+{
+	TSharedPtr<FPhysicsAssetEditor> PhysicsAssetEditor = PhysicsAssetEditorPtr.Pin();
+	PhysicsAssetEditor->RegisterTabSpawners(InTabManager.ToSharedRef());
+	PhysicsAssetEditor->PushTabFactories(TabFactories);
+
+	FApplicationMode::RegisterTabFactories(InTabManager);
+}
+
+#undef LOCTEXT_NAMESPACE
\ No newline at end of file
