Copying //UE4/Dev-AnimPhys to //UE4/Dev-Main (Source: //UE4/Dev-AnimPhys @ 3624379)

diff --git a/Engine/Source/Editor/PhysicsAssetEditor/Private/PhysicsAssetEditorSharedData.cpp b/Engine/Source/Editor/PhysicsAssetEditor/Private/PhysicsAssetEditorSharedData.cpp
new file mode 100644
--- /dev/null
+++ b/Engine/Source/Editor/PhysicsAssetEditor/Private/PhysicsAssetEditorSharedData.cpp
@@ -0,0 +1,1741 @@
+// Copyright 1998-2017 Epic Games, Inc. All Rights Reserved.
+
+#include "PhysicsAssetEditorSharedData.h"
+#include "PhysicsEngine/PhysicsHandleComponent.h"
+#include "PhysicsEngine/RigidBodyIndexPair.h"
+#include "Misc/MessageDialog.h"
+#include "Modules/ModuleManager.h"
+#include "UObject/UObjectIterator.h"
+#include "Widgets/SWindow.h"
+#include "Components/StaticMeshComponent.h"
+#include "Engine/SkeletalMesh.h"
+#include "Components/SkeletalMeshComponent.h"
+#include "Preferences/PhysicsAssetEditorOptions.h"
+#include "Engine/StaticMesh.h"
+#include "Engine/CollisionProfile.h"
+#include "Editor.h"
+#include "PhysicsAssetEditorModule.h"
+#include "EditorSupportDelegates.h"
+#include "ScopedTransaction.h"
+#include "PhysicsAssetEditorSkeletalMeshComponent.h"
+#include "MeshUtilities.h"
+#include "PhysicsEngine/BoxElem.h"
+#include "PhysicsEngine/ConstraintInstance.h"
+#include "PhysicsEngine/PhysicsConstraintTemplate.h"
+#include "PhysicsEngine/PhysicalAnimationComponent.h"
+#include "PhysicsEngine/PhysicsAsset.h"
+#include "PhysicsAssetEditorAnimInstance.h"
+#include "IPersonaPreviewScene.h"
+#include "PhysicsPublic.h"
+#include "PhysicsAssetGenerationSettings.h"
+#include "IDetailsView.h"
+#include "PropertyEditorModule.h"
+#include "SButton.h"
+#include "STextBlock.h"
+
+#define LOCTEXT_NAMESPACE "PhysicsAssetEditorShared"
+
+FPhysicsAssetEditorSharedData::FPhysicsAssetEditorSharedData()
+	: COMRenderColor(255,255,100)
+	, CopiedBodySetup(NULL)
+	, CopiedConstraintTemplate(NULL)
+	, InsideSelChange(0)
+{
+	// Editor variables
+	bShowCOM = false;
+
+	bRunningSimulation = false;
+	bNoGravitySimulation = false;
+
+	bManipulating = false;
+	
+	// Construct mouse handle
+	MouseHandle = NewObject<UPhysicsHandleComponent>();
+
+	// Construct sim options.
+	EditorOptions = NewObject<UPhysicsAssetEditorOptions>(GetTransientPackage(), TEXT("EditorOptions"));
+	check(EditorOptions);
+
+	EditorOptions->LoadConfig();
+}
+
+FPhysicsAssetEditorSharedData::~FPhysicsAssetEditorSharedData()
+{
+
+}
+
+void FPhysicsAssetEditorSharedData::Initialize(const TSharedRef<IPersonaPreviewScene>& InPreviewScene)
+{
+	PreviewScene = InPreviewScene;
+
+	EditorSkelComp = nullptr;
+	PhysicalAnimationComponent = nullptr;
+
+	// Look for body setups with no shapes (how does this happen?).
+	// If we find one- just bang on a default box.
+	bool bFoundEmptyShape = false;
+	for (int32 i = 0; i <PhysicsAsset->SkeletalBodySetups.Num(); ++i)
+	{
+		UBodySetup* BodySetup = PhysicsAsset->SkeletalBodySetups[i];
+		if (BodySetup->AggGeom.GetElementCount() == 0)
+		{
+			FKBoxElem BoxElem;
+			BoxElem.SetTransform(FTransform::Identity);
+			BoxElem.X = 15.f;
+			BoxElem.Y = 15.f;
+			BoxElem.Z = 15.f;
+			BodySetup->AggGeom.BoxElems.Add(BoxElem);
+			check(BodySetup->AggGeom.BoxElems.Num() == 1);
+
+			bFoundEmptyShape = true;
+		}
+	}
+
+	// Pop up a warning about what we did.
+	if (bFoundEmptyShape)
+	{
+		FMessageDialog::Open(EAppMsgType::Ok, NSLOCTEXT("UnrealEd", "EmptyBodyFound", "Bodies was found with no primitives!\nThey have been reset to have a box."));
+	}
+
+	IMeshUtilities& MeshUtilities = FModuleManager::Get().LoadModuleChecked<IMeshUtilities>("MeshUtilities");
+	// Used for viewing bone influences, resetting bone geometry etc.
+	USkeletalMesh* EditorSkelMesh = PhysicsAsset->GetPreviewMesh();
+	if(EditorSkelMesh)
+	{
+		MeshUtilities.CalcBoneVertInfos(EditorSkelMesh, DominantWeightBoneInfos, true);
+		MeshUtilities.CalcBoneVertInfos(EditorSkelMesh, AnyWeightBoneInfos, false);
+
+		// Ensure PhysicsAsset mass properties are up to date.
+		PhysicsAsset->UpdateBoundsBodiesArray();
+
+		// Check if there are any bodies in the Asset which do not have bones in the skeletal mesh.
+		// If so, put up a warning.
+		TArray<int32> MissingBodyIndices;
+		FString BoneNames;
+		for (int32 i = 0; i <PhysicsAsset->SkeletalBodySetups.Num(); ++i)
+		{
+			FName BoneName = PhysicsAsset->SkeletalBodySetups[i]->BoneName;
+			int32 BoneIndex = EditorSkelMesh->RefSkeleton.FindBoneIndex(BoneName);
+			if (BoneIndex == INDEX_NONE)
+			{
+				MissingBodyIndices.Add( i );
+				BoneNames += FString::Printf(TEXT("\t%s\n"), *BoneName.ToString());
+			}
+		}
+
+		const FText MissingBodyMsg = FText::Format( LOCTEXT( "MissingBones", "The following Bodies are in the PhysicsAsset, but have no corresponding bones in the SkeletalMesh.\nClick OK to delete them, or Cancel to ignore.\n\n{0}" ), FText::FromString( BoneNames ) );
+
+		if ( MissingBodyIndices.Num() )
+		{
+			if ( FMessageDialog::Open( EAppMsgType::OkCancel, MissingBodyMsg ) == EAppReturnType::Ok )
+			{
+				// Delete the bodies with no associated bones
+
+				const FScopedTransaction Transaction( LOCTEXT( "DeleteUnusedPhysicsBodies", "Delete Physics Bodies With No Bones" ) );
+				PhysicsAsset->SetFlags(RF_Transactional);
+				PhysicsAsset->Modify();
+
+				// Iterate backwards, as PhysicsAsset->SkeletalBodySetups is a TArray and UE4 containers don't support remove_if()
+				for ( int32 i = MissingBodyIndices.Num() - 1; i >= 0; --i )
+				{
+					DeleteBody( MissingBodyIndices[i], false );
+				}
+			}
+		}
+	}
+
+	// Support undo/redo
+	PhysicsAsset->SetFlags(RF_Transactional);
+
+	ClearSelectedBody();
+	ClearSelectedConstraints();
+}
+
+void FPhysicsAssetEditorSharedData::CachePreviewMesh()
+{
+	USkeletalMesh* PreviewMesh = PhysicsAsset->PreviewSkeletalMesh.LoadSynchronous();
+
+	if (PreviewMesh == nullptr)
+	{
+		// Fall back to the default skeletal mesh in the EngineMeshes package.
+		// This is statically loaded as the package is likely not fully loaded
+		// (otherwise, it would have been found in the above iteration).
+		PreviewMesh = (USkeletalMesh*)StaticLoadObject(USkeletalMesh::StaticClass(), NULL, TEXT("/Engine/EngineMeshes/SkeletalCube.SkeletalCube"), NULL, LOAD_None, NULL);
+		check(PreviewMesh);
+
+		PhysicsAsset->PreviewSkeletalMesh = PreviewMesh;
+
+		FMessageDialog::Open(EAppMsgType::Ok, FText::Format(
+				LOCTEXT("Error_PhysicsAssetHasNoSkelMesh", "Warning: Physics Asset has no skeletal mesh assigned.\nFor now, a simple default skeletal mesh ({0}) will be used.\nYou can fix this by opening the asset and choosing another skeletal mesh from the toolbar."),
+				FText::FromString(PreviewMesh->GetFullName())));
+	}
+}
+
+void FPhysicsAssetEditorSharedData::CopyConstraintProperties(UPhysicsConstraintTemplate * FromConstraintSetup, UPhysicsConstraintTemplate * ToConstraintSetup)
+{
+	ToConstraintSetup->Modify();
+	FConstraintInstance OldInstance = ToConstraintSetup->DefaultInstance;
+	ToConstraintSetup->DefaultInstance.CopyConstraintParamsFrom(&FromConstraintSetup->DefaultInstance);
+
+	// recover certain data that we'd like to keep - i.e. bone indices those still should stay.  
+	// frame position offsets taken from old, but frame orientations are taken from new source
+	ToConstraintSetup->DefaultInstance.ConstraintIndex = OldInstance.ConstraintIndex;
+	ToConstraintSetup->DefaultInstance.ConstraintData = OldInstance.ConstraintData;
+	ToConstraintSetup->DefaultInstance.JointName = OldInstance.JointName;
+	ToConstraintSetup->DefaultInstance.ConstraintBone1 = OldInstance.ConstraintBone1;
+	ToConstraintSetup->DefaultInstance.ConstraintBone2 = OldInstance.ConstraintBone2;
+	ToConstraintSetup->DefaultInstance.Pos1 = OldInstance.Pos1;
+	ToConstraintSetup->DefaultInstance.Pos2 = OldInstance.Pos2;
+}
+
+struct FMirrorInfo
+{
+	FName BoneName;
+	int32 BoneIndex;
+	int32 BodyIndex;
+	int32 ConstraintIndex;
+	FMirrorInfo()
+	{
+		BoneIndex = INDEX_NONE;
+		BodyIndex = INDEX_NONE;
+		ConstraintIndex = INDEX_NONE;
+		BoneName = NAME_None;
+	}
+};
+
+void FPhysicsAssetEditorSharedData::Mirror()
+{
+	USkeletalMesh* EditorSkelMesh = PhysicsAsset->GetPreviewMesh();
+	if(EditorSkelMesh)
+	{
+
+		TArray<FMirrorInfo> MirrorInfos;
+
+		for (const FSelection& Selection : SelectedBodies)
+		{
+			MirrorInfos.AddUninitialized();
+			FMirrorInfo & MirrorInfo = MirrorInfos[MirrorInfos.Num() - 1];
+			MirrorInfo.BoneName = PhysicsAsset->SkeletalBodySetups[Selection.Index]->BoneName;
+			MirrorInfo.BodyIndex = Selection.Index;
+			MirrorInfo.ConstraintIndex = PhysicsAsset->FindConstraintIndex(MirrorInfo.BoneName);
+		}
+
+		for (const FSelection& Selection : SelectedConstraints)
+		{
+			MirrorInfos.AddUninitialized();
+			FMirrorInfo & MirrorInfo = MirrorInfos[MirrorInfos.Num() - 1];
+			MirrorInfo.BoneName = PhysicsAsset->ConstraintSetup[Selection.Index]->DefaultInstance.ConstraintBone1;
+			MirrorInfo.BodyIndex = PhysicsAsset->FindBodyIndex(MirrorInfo.BoneName);
+			MirrorInfo.ConstraintIndex = Selection.Index;
+		}
+
+		for (FMirrorInfo & MirrorInfo : MirrorInfos)	//mirror all selected bodies/constraints
+		{
+			int32 BoneIndex = EditorSkelMesh->RefSkeleton.FindBoneIndex(MirrorInfo.BoneName);
+
+			int32 MirrorBoneIndex = PhysicsAsset->FindMirroredBone(EditorSkelMesh, BoneIndex);
+			if (MirrorBoneIndex != INDEX_NONE)
+			{
+				UBodySetup * SrcBody = PhysicsAsset->SkeletalBodySetups[MirrorInfo.BodyIndex];
+				const FScopedTransaction Transaction(NSLOCTEXT("PhysicsAssetEditor", "MirrorBody", "MirrorBody"));
+				MakeNewBody(MirrorBoneIndex, false);
+
+				int32 MirrorBodyIndex = PhysicsAsset->FindControllingBodyIndex(EditorSkelMesh, MirrorBoneIndex);
+
+				UBodySetup * DestBody = PhysicsAsset->SkeletalBodySetups[MirrorBodyIndex];
+				DestBody->Modify();
+				DestBody->CopyBodyPropertiesFrom(SrcBody);
+				FQuat ArtistMirrorConvention(0,0,1,0);   // how Epic Maya artists rig the right and left orientation differently.  todo: perhaps move to cvar 
+				for (FKSphylElem& Sphyl : DestBody->AggGeom.SphylElems)
+				{
+					Sphyl.Rotation	= (ArtistMirrorConvention*Sphyl.Rotation.Quaternion()).Rotator();
+					Sphyl.Center	= ArtistMirrorConvention.RotateVector(Sphyl.Center);
+				}
+				for (FKBoxElem& Box : DestBody->AggGeom.BoxElems)
+				{
+					Box.Rotation	= (ArtistMirrorConvention*Box.Rotation.Quaternion()).Rotator();
+					Box.Center      = ArtistMirrorConvention.RotateVector(Box.Center);
+				}
+				for (FKSphereElem& Sphere : DestBody->AggGeom.SphereElems)
+				{
+					Sphere.Center = ArtistMirrorConvention.RotateVector(Sphere.Center);
+				}
+				int32 MirrorConstraintIndex = PhysicsAsset->FindConstraintIndex(DestBody->BoneName);
+				UPhysicsConstraintTemplate * FromConstraint = PhysicsAsset->ConstraintSetup[MirrorInfo.ConstraintIndex];
+				UPhysicsConstraintTemplate * ToConstraint = PhysicsAsset->ConstraintSetup[MirrorConstraintIndex];
+				CopyConstraintProperties(FromConstraint, ToConstraint);
+			}
+		}
+	}
+}
+
+EPhysicsAssetEditorRenderMode FPhysicsAssetEditorSharedData::GetCurrentMeshViewMode(bool bSimulation)
+{
+	if (bSimulation)
+	{
+		return EditorOptions->SimulationMeshViewMode;
+	}
+	else
+	{
+		return EditorOptions->MeshViewMode;
+	}
+}
+
+EPhysicsAssetEditorRenderMode FPhysicsAssetEditorSharedData::GetCurrentCollisionViewMode(bool bSimulation)
+{
+	if (bSimulation)
+	{
+		return EditorOptions->SimulationCollisionViewMode;
+	}
+	else
+	{
+		return EditorOptions->CollisionViewMode;
+	}
+}
+
+EPhysicsAssetEditorConstraintViewMode FPhysicsAssetEditorSharedData::GetCurrentConstraintViewMode(bool bSimulation)
+{
+	if (bSimulation)
+	{
+		return EditorOptions->SimulationConstraintViewMode;
+	}
+	else
+	{
+		return EditorOptions->ConstraintViewMode;
+	}
+}
+
+void FPhysicsAssetEditorSharedData::HitBone(int32 BodyIndex, EAggCollisionShape::Type PrimType, int32 PrimIndex, bool bGroupSelect /* = false*/)
+{
+	if (!bRunningSimulation)
+	{
+		FPhysicsAssetEditorSharedData::FSelection Selection(BodyIndex, PrimType, PrimIndex);
+		if(bGroupSelect)
+		{
+			if(IsBodySelected(Selection))
+			{
+				SetSelectedBody(Selection, false);
+			}
+			else
+			{
+				SetSelectedBody(Selection, true);
+			}
+		}
+		else
+		{
+			ClearSelectedBody();
+			SetSelectedBody(Selection, true);
+		}
+	}
+}
+
+void FPhysicsAssetEditorSharedData::HitConstraint(int32 ConstraintIndex, bool bGroupSelect)
+{
+	if (!bRunningSimulation)
+	{
+		if(bGroupSelect)
+		{
+			if(IsConstraintSelected(ConstraintIndex))
+			{
+				SetSelectedConstraint(ConstraintIndex, false);
+			}
+			else
+			{
+				SetSelectedConstraint(ConstraintIndex, true);
+			}
+		}
+		else
+		{
+			ClearSelectedConstraints();
+			SetSelectedConstraint(ConstraintIndex, true);
+		}
+	}
+}
+
+void FPhysicsAssetEditorSharedData::RefreshPhysicsAssetChange(const UPhysicsAsset* InPhysAsset)
+{
+	if (InPhysAsset)
+	{
+		InPhysAsset->RefreshPhysicsAssetChange();
+
+		// Broadbcast delegate
+		FPhysicsDelegates::OnPhysicsAssetChanged.Broadcast(InPhysAsset);
+
+		FEditorSupportDelegates::RedrawAllViewports.Broadcast();
+		// since we recreate physicsstate, a lot of transient state data will be gone
+		// so have to turn simulation off again. 
+		// ideally maybe in the future, we'll fix it by controlling tick?
+		EnableSimulation(false);
+		EditorSkelComp->RecreatePhysicsState();
+		ForceDisableSimulation();
+	}
+}
+
+void FPhysicsAssetEditorSharedData::SetSelectedBodyAnyPrim(int32 BodyIndex, bool bSelected)
+{
+	if (BodyIndex == INDEX_NONE)
+	{
+		ClearSelectedBody();
+		return;
+	}
+	
+	UBodySetup* BodySetup = PhysicsAsset->SkeletalBodySetups[BodyIndex];
+	check(BodySetup);
+
+	if (BodySetup->AggGeom.SphereElems.Num() > 0)
+	{
+		FSelection Selection(BodyIndex, EAggCollisionShape::Sphere, 0);
+		SetSelectedBody(Selection, bSelected);
+	}
+	else if (BodySetup->AggGeom.BoxElems.Num() > 0)
+	{
+		FSelection Selection(BodyIndex, EAggCollisionShape::Box, 0);
+		SetSelectedBody(Selection, bSelected);
+	}
+	else if (BodySetup->AggGeom.SphylElems.Num() > 0)
+	{
+		FSelection Selection(BodyIndex, EAggCollisionShape::Sphyl, 0);
+		SetSelectedBody(Selection, bSelected);
+	}
+	else if (BodySetup->AggGeom.ConvexElems.Num() > 0)
+	{
+		FSelection Selection(BodyIndex, EAggCollisionShape::Convex, 0);
+		SetSelectedBody(Selection, bSelected);
+	}
+	else
+	{
+		UE_LOG(LogPhysicsAssetEditor, Fatal, TEXT("Body Setup with No Primitives!")); 
+	}
+}
+
+void FPhysicsAssetEditorSharedData::ClearSelectedBody()
+{
+	SelectedBodies.Empty();
+	SelectedConstraints.Empty();
+
+	SelectionChangedEvent.Broadcast(SelectedBodies, SelectedConstraints);
+}
+
+void FPhysicsAssetEditorSharedData::SetSelectedBody(const FSelection& Body, bool bSelected)
+{
+	if(InsideSelChange)
+	{
+		return;
+	}
+
+	if(bSelected)
+	{
+		SelectedBodies.AddUnique(Body);
+	}
+	else
+	{
+		SelectedBodies.Remove(Body);
+	}
+
+	SelectionChangedEvent.Broadcast(SelectedBodies, SelectedConstraints);
+
+	if(!GetSelectedBody())
+	{
+		return;
+	}
+
+	UpdateNoCollisionBodies();
+	++InsideSelChange;
+	PreviewChangedEvent.Broadcast();
+	--InsideSelChange;
+}
+
+bool FPhysicsAssetEditorSharedData::IsBodySelected(const FSelection& Body) const
+{
+	return SelectedBodies.Contains(Body);
+}
+
+void FPhysicsAssetEditorSharedData::UpdateNoCollisionBodies()
+{
+	NoCollisionBodies.Empty();
+
+	// Query disable table with selected body and every other body.
+	for (int32 i = 0; i <PhysicsAsset->SkeletalBodySetups.Num(); ++i)
+	{
+		// Add any bodies with bNoCollision
+		if (PhysicsAsset->SkeletalBodySetups[i]->DefaultInstance.GetCollisionEnabled() == ECollisionEnabled::NoCollision)
+		{
+			NoCollisionBodies.Add(i);
+		}
+		else if (GetSelectedBody() && i != GetSelectedBody()->Index)
+		{
+			// Add this body if it has disabled collision with selected.
+			FRigidBodyIndexPair Key(i, GetSelectedBody()->Index);
+
+			if (PhysicsAsset->SkeletalBodySetups[GetSelectedBody()->Index]->DefaultInstance.GetCollisionEnabled() == ECollisionEnabled::NoCollision ||
+				PhysicsAsset->CollisionDisableTable.Find(Key))
+			{
+				NoCollisionBodies.Add(i);
+			}
+		}
+	}
+}
+
+void FPhysicsAssetEditorSharedData::ClearSelectedConstraints()
+{
+	if(InsideSelChange)
+	{
+		return;
+	}
+
+	SelectedBodies.Empty();
+	SelectedConstraints.Empty();
+
+	SelectionChangedEvent.Broadcast(SelectedBodies, SelectedConstraints);
+
+	++InsideSelChange;
+	PreviewChangedEvent.Broadcast();
+	--InsideSelChange;
+}
+
+void FPhysicsAssetEditorSharedData::SetSelectedConstraint(int32 ConstraintIndex, bool bSelected)
+{
+	if(InsideSelChange)
+	{
+		return;
+	}
+
+	if(ConstraintIndex != INDEX_NONE)
+	{
+		FSelection Constraint(ConstraintIndex, EAggCollisionShape::Unknown, INDEX_NONE);
+		if(bSelected)
+		{
+			SelectedConstraints.AddUnique(Constraint);
+		}
+		else
+		{
+			SelectedConstraints.Remove(Constraint);
+		}
+
+		SelectionChangedEvent.Broadcast(SelectedBodies, SelectedConstraints);
+
+		++InsideSelChange;
+		PreviewChangedEvent.Broadcast();
+		--InsideSelChange;
+	}
+}
+
+bool FPhysicsAssetEditorSharedData::IsConstraintSelected(int32 ConstraintIndex) const
+{
+	FSelection Constraint(ConstraintIndex, EAggCollisionShape::Unknown, INDEX_NONE);
+	return SelectedConstraints.Contains(Constraint);
+}
+
+void FPhysicsAssetEditorSharedData::SetCollisionBetweenSelected(bool bEnableCollision)
+{
+	if (bRunningSimulation || SelectedBodies.Num() == 0)
+	{
+		return;
+	}
+
+	PhysicsAsset->Modify();
+
+	for(int32 i=0; i<SelectedBodies.Num(); ++i)
+	{
+		for(int32 j=i+1; j<SelectedBodies.Num(); ++j)
+		{
+			if(bEnableCollision)
+			{
+				PhysicsAsset->EnableCollision(SelectedBodies[i].Index, SelectedBodies[j].Index);
+			}else
+			{
+				PhysicsAsset->DisableCollision(SelectedBodies[i].Index, SelectedBodies[j].Index);
+			}
+
+		}
+	}
+
+
+	UpdateNoCollisionBodies();
+
+	PreviewChangedEvent.Broadcast();
+}
+
+bool FPhysicsAssetEditorSharedData::CanSetCollisionBetweenSelected(bool bEnableCollision) const
+{
+	if (bRunningSimulation || SelectedBodies.Num() == 0)
+	{
+		return false;
+	}
+
+	for(int32 i=0; i<SelectedBodies.Num(); ++i)
+	{
+		for(int32 j=i+1; j<SelectedBodies.Num(); ++j)
+		{
+			if(PhysicsAsset->IsCollisionEnabled(SelectedBodies[i].Index, SelectedBodies[j].Index) != bEnableCollision)
+			{
+				return true;
+			}
+		}
+	}
+
+	return false;
+}
+
+void FPhysicsAssetEditorSharedData::SetCollisionBetweenSelectedAndAll(bool bEnableCollision)
+{
+	if (bRunningSimulation || SelectedBodies.Num() == 0)
+	{
+		return;
+	}
+
+	PhysicsAsset->Modify();
+
+	for(int32 i=0; i<SelectedBodies.Num(); ++i)
+	{
+		for(int32 j = 0; j < PhysicsAsset->SkeletalBodySetups.Num(); ++j)
+		{
+			if(bEnableCollision)
+			{
+				PhysicsAsset->EnableCollision(SelectedBodies[i].Index, j);
+			}
+			else
+			{
+				PhysicsAsset->DisableCollision(SelectedBodies[i].Index, j);
+			}
+
+		}
+	}
+
+	UpdateNoCollisionBodies();
+
+	PreviewChangedEvent.Broadcast();
+}
+
+bool FPhysicsAssetEditorSharedData::CanSetCollisionBetweenSelectedAndAll(bool bEnableCollision) const
+{
+	if (bRunningSimulation || SelectedBodies.Num() == 0)
+	{
+		return false;
+	}
+
+	for(int32 i=0; i<SelectedBodies.Num(); ++i)
+	{
+		for(int32 j = 0; j < PhysicsAsset->SkeletalBodySetups.Num(); ++j)
+		{
+			if(PhysicsAsset->IsCollisionEnabled(SelectedBodies[i].Index, j) != bEnableCollision)
+			{
+				return true;
+			}
+		}
+	}
+
+	return false;
+}
+
+void FPhysicsAssetEditorSharedData::SetCollisionBetween(int32 Body1Index, int32 Body2Index, bool bEnableCollision)
+{
+	if (bRunningSimulation)
+	{
+		return;
+	}
+
+	PhysicsAsset->Modify();
+
+	if (Body1Index != INDEX_NONE && Body2Index != INDEX_NONE && Body1Index != Body2Index)
+	{
+		if (bEnableCollision)
+		{
+			PhysicsAsset->EnableCollision(Body1Index, Body2Index);
+		}
+		else
+		{
+			PhysicsAsset->DisableCollision(Body1Index, Body2Index);
+		}
+
+		UpdateNoCollisionBodies();
+	}
+
+	PreviewChangedEvent.Broadcast();
+}
+
+void FPhysicsAssetEditorSharedData::CopyBody()
+{
+	check(SelectedBodies.Num() == 1);
+
+	CopiedBodySetup = PhysicsAsset->SkeletalBodySetups[GetSelectedBody()->Index];
+}
+
+void FPhysicsAssetEditorSharedData::PasteBodyProperties()
+{
+	// Can't do this while simulating!
+	if (bRunningSimulation)
+	{
+		return;
+	}
+
+	// Must have two valid bodies (which are different)
+	if(CopiedBodySetup == NULL)
+	{
+		return;
+	}
+
+	const FScopedTransaction Transaction( NSLOCTEXT("PhysicsAssetEditor", "PasteBodyProperties", "Paste Body Properties") );
+
+	for(int32 i=0; i<SelectedBodies.Num(); ++i)
+	{
+		UBodySetup* ToBodySetup = PhysicsAsset->SkeletalBodySetups[SelectedBodies[i].Index];
+		UBodySetup* FromBodySetup = CopiedBodySetup;
+		ToBodySetup->Modify();
+		ToBodySetup->CopyBodyPropertiesFrom(FromBodySetup);
+	}
+	
+	ClearSelectedBody();	//paste can change the primitives on our selected bodies. There's probably a way to properly update this, but for now just deselect
+	PreviewChangedEvent.Broadcast();
+}
+
+bool FPhysicsAssetEditorSharedData::WeldSelectedBodies(bool bWeld /* = true */)
+{
+	bool bCanWeld = false;
+	if (bRunningSimulation)
+	{
+		return false;
+	}
+
+	if(SelectedBodies.Num() <= 1)
+	{
+		return false;
+	}
+
+	USkeletalMesh* EditorSkelMesh = PhysicsAsset->GetPreviewMesh();
+	if(EditorSkelMesh == nullptr)
+	{
+		return false;
+	}
+
+	//we only support two body weld
+	int BodyIndex0 = 0;
+	int BodyIndex1 = INDEX_NONE;
+
+	for(int32 i=1; i<SelectedBodies.Num(); ++i)
+	{
+		if(SelectedBodies[BodyIndex0].Index == SelectedBodies[i].Index)
+		{
+			continue;
+		}
+
+		if(BodyIndex1== INDEX_NONE)
+		{
+			BodyIndex1 = i;
+		}else
+		{
+			if(SelectedBodies[BodyIndex1].Index != SelectedBodies[i].Index)
+			{
+				return false;
+			}
+		}
+	}
+
+	//need to weld bodies not primitives
+	if(BodyIndex1 == INDEX_NONE)
+	{
+		return false;
+	}
+
+	const FSelection& Body0 = SelectedBodies[BodyIndex0];
+	const FSelection& Body1 = SelectedBodies[BodyIndex1];
+
+	FName Bone0Name = PhysicsAsset->SkeletalBodySetups[Body0.Index]->BoneName;
+	int32 Bone0Index = EditorSkelMesh->RefSkeleton.FindBoneIndex(Bone0Name);
+	check(Bone0Index != INDEX_NONE);
+
+	FName Bone1Name = PhysicsAsset->SkeletalBodySetups[Body1.Index]->BoneName;
+	int32 Bone1Index = EditorSkelMesh->RefSkeleton.FindBoneIndex(Bone1Name);
+	check(Bone1Index != INDEX_NONE);
+
+	int32 Bone0ParentIndex = EditorSkelMesh->RefSkeleton.GetParentIndex(Bone0Index);
+	int32 Bone1ParentIndex = EditorSkelMesh->RefSkeleton.GetParentIndex(Bone1Index);
+
+	int ParentBodyIndex = INDEX_NONE;
+	int ChildBodyIndex = INDEX_NONE;
+	FName ParentBoneName;
+	EAggCollisionShape::Type ParentPrimitiveType = EAggCollisionShape::Unknown;
+	EAggCollisionShape::Type ChildPrimitiveType = EAggCollisionShape::Unknown;
+	int32 ParentPrimitiveIndex = INDEX_NONE;
+	int32 ChildPrimitiveIndex = INDEX_NONE;
+
+	if (PhysicsAsset->FindControllingBodyIndex(EditorSkelMesh, Bone1ParentIndex) == Body0.Index)
+	{
+		ParentBodyIndex = Body0.Index;
+		ParentBoneName = Bone0Name;
+		ChildBodyIndex = Body1.Index;
+		ParentPrimitiveType = Body0.PrimitiveType;
+		ChildPrimitiveType = Body1.PrimitiveType;
+		ParentPrimitiveIndex = Body0.PrimitiveIndex;
+		//Child geoms get appended so just add it. This is kind of a hack but this whole indexing scheme needs to be rewritten anyway
+		ChildPrimitiveIndex = Body1.PrimitiveIndex + PhysicsAsset->SkeletalBodySetups[Body0.Index]->AggGeom.GetElementCount(ChildPrimitiveType);
+
+		bCanWeld = true;
+	}else if(PhysicsAsset->FindControllingBodyIndex(EditorSkelMesh, Bone0ParentIndex) == Body1.Index)
+	{
+		ParentBodyIndex = Body1.Index;
+		ParentBoneName = Bone1Name;
+		ChildBodyIndex = Body0.Index;
+		ParentPrimitiveType = Body1.PrimitiveType;
+		ChildPrimitiveType = Body0.PrimitiveType;
+		ParentPrimitiveIndex = Body1.PrimitiveIndex;
+		//Child geoms get appended so just add it. This is kind of a hack but this whole indexing scheme needs to be rewritten anyway
+		ChildPrimitiveIndex = Body0.PrimitiveIndex + PhysicsAsset->SkeletalBodySetups[Body1.Index]->AggGeom.GetElementCount(ChildPrimitiveType);
+
+		bCanWeld = true;
+	}
+
+	//function is used for the action and the check
+	if(bWeld == false)
+	{
+		return bCanWeld;
+	}
+
+	check(ParentBodyIndex != INDEX_NONE);
+	check(ChildBodyIndex != INDEX_NONE);
+
+	{
+		const FScopedTransaction Transaction( NSLOCTEXT("UnrealEd", "WeldBodies", "Weld Bodies") );
+
+		// .. the asset itself..
+		PhysicsAsset->Modify();
+
+		// .. the parent and child bodies..
+		PhysicsAsset->SkeletalBodySetups[ParentBodyIndex]->Modify();
+		PhysicsAsset->SkeletalBodySetups[ChildBodyIndex]->Modify();
+
+		// .. and any constraints of the 'child' body..
+		TArray<int32>	Constraints;
+		PhysicsAsset->BodyFindConstraints(ChildBodyIndex, Constraints);
+
+		for (int32 i = 0; i <Constraints.Num(); ++i)
+		{
+			int32 ConstraintIndex = Constraints[i];
+			PhysicsAsset->ConstraintSetup[ConstraintIndex]->Modify();
+		}
+
+		// Do the actual welding
+		FPhysicsAssetUtils::WeldBodies(PhysicsAsset, ParentBodyIndex, ChildBodyIndex, EditorSkelComp);
+	}
+
+	// update the tree
+	HierarchyChangedEvent.Broadcast();
+
+	// Body index may have changed, so we re-find it.
+	int32 BodyIndex = PhysicsAsset->FindBodyIndex(ParentBoneName);
+	FSelection SelectionParent(BodyIndex, ParentPrimitiveType, ParentPrimitiveIndex);
+	SetSelectedBody(SelectionParent, true); // This redraws the viewport as well...
+
+	FSelection SelectionChild(BodyIndex, ChildPrimitiveType, ChildPrimitiveIndex);
+	SetSelectedBody(SelectionChild, true); // This redraws the viewport as well...
+
+	// Just to be safe - deselect any selected constraints
+	ClearSelectedConstraints();
+	RefreshPhysicsAssetChange(PhysicsAsset);
+	return true;
+}
+
+
+void FPhysicsAssetEditorSharedData::InitConstraintSetup(UPhysicsConstraintTemplate* ConstraintSetup, int32 ChildBodyIndex, int32 ParentBodyIndex)
+{
+	check(ConstraintSetup);
+
+	ConstraintSetup->Modify(false);
+
+	UBodySetup* ChildBodySetup = PhysicsAsset->SkeletalBodySetups[ ChildBodyIndex ];
+	UBodySetup* ParentBodySetup = PhysicsAsset->SkeletalBodySetups[ ParentBodyIndex ];
+	check(ChildBodySetup && ParentBodySetup);
+
+	// Place joint at origin of child
+	ConstraintSetup->DefaultInstance.ConstraintBone1 = ChildBodySetup->BoneName;
+	ConstraintSetup->DefaultInstance.ConstraintBone2 = ParentBodySetup->BoneName;
+	SnapConstraintToBone(ConstraintSetup->DefaultInstance);
+
+	ConstraintSetup->SetDefaultProfile(ConstraintSetup->DefaultInstance);
+
+	// Disable collision between constrained bodies by default.
+	SetCollisionBetween(ChildBodyIndex, ParentBodyIndex, false);
+}
+
+void FPhysicsAssetEditorSharedData::MakeNewBody(int32 NewBoneIndex, bool bAutoSelect)
+{
+	USkeletalMesh* EditorSkelMesh = PhysicsAsset->GetPreviewMesh();
+	if(EditorSkelMesh == nullptr)
+	{
+		return;
+	}
+
+	PhysicsAsset->Modify();
+
+	FName NewBoneName = EditorSkelMesh->RefSkeleton.GetBoneName(NewBoneIndex);
+
+	// If this body is already physical, remove the current body
+	int32 NewBodyIndex = PhysicsAsset->FindBodyIndex(NewBoneName);
+	if (NewBodyIndex != INDEX_NONE)
+	{
+		DeleteBody(NewBodyIndex, false);
+	}
+
+	// Find body that currently controls this bone.
+	int32 ParentBodyIndex = PhysicsAsset->FindControllingBodyIndex(EditorSkelMesh, NewBoneIndex);
+
+	// Create the physics body.
+	NewBodyIndex = FPhysicsAssetUtils::CreateNewBody(PhysicsAsset, NewBoneName);
+	UBodySetup* BodySetup = PhysicsAsset->SkeletalBodySetups[ NewBodyIndex ];
+	check(BodySetup->BoneName == NewBoneName);
+	
+	BodySetup->Modify();
+
+	bool bCreatedBody = false;
+	// Create a new physics body for this bone.
+	const FPhysAssetCreateParams& NewBodyData = GetDefault<UPhysicsAssetGenerationSettings>()->CreateParams;
+	if (NewBodyData.VertWeight == EVW_DominantWeight)
+	{
+		bCreatedBody = FPhysicsAssetUtils::CreateCollisionFromBone(BodySetup, EditorSkelMesh, NewBoneIndex, NewBodyData, DominantWeightBoneInfos[NewBoneIndex]);
+	}
+	else
+	{
+		bCreatedBody = FPhysicsAssetUtils::CreateCollisionFromBone(BodySetup, EditorSkelMesh, NewBoneIndex, NewBodyData, AnyWeightBoneInfos[NewBoneIndex]);
+	}
+
+	if (bCreatedBody == false)
+	{
+		FPhysicsAssetUtils::DestroyBody(PhysicsAsset, NewBodyIndex);
+		return;
+	}
+
+	// Check if the bone of the new body has any physical children bones
+	for (int32 i = 0; i < EditorSkelMesh->RefSkeleton.GetRawBoneNum(); ++i)
+	{
+		if (EditorSkelMesh->RefSkeleton.BoneIsChildOf(i, NewBoneIndex))
+		{
+			const int32 ChildBodyIndex = PhysicsAsset->FindBodyIndex(EditorSkelMesh->RefSkeleton.GetBoneName(i));
+			
+			// If the child bone is physical, it may require fixing up in regards to constraints
+			if (ChildBodyIndex != INDEX_NONE)
+			{
+				UBodySetup* ChildBody = PhysicsAsset->SkeletalBodySetups[ ChildBodyIndex ];
+				check(ChildBody);
+
+				int32 ConstraintIndex = PhysicsAsset->FindConstraintIndex(ChildBody->BoneName);
+				
+				// If the child body is not constrained already, create a new constraint between
+				// the child body and the new body
+				if (ConstraintIndex == INDEX_NONE)
+				{
+					ConstraintIndex = FPhysicsAssetUtils::CreateNewConstraint(PhysicsAsset, ChildBody->BoneName);
+					check(ConstraintIndex != INDEX_NONE);
+				}
+				// If there's a pre-existing constraint, see if it needs to be fixed up
+				else
+				{
+					UPhysicsConstraintTemplate* ExistingConstraintSetup = PhysicsAsset->ConstraintSetup[ ConstraintIndex ];
+					check(ExistingConstraintSetup);
+					
+					const int32 ExistingConstraintBoneIndex = EditorSkelMesh->RefSkeleton.FindBoneIndex(ExistingConstraintSetup->DefaultInstance.ConstraintBone2);
+					check(ExistingConstraintBoneIndex != INDEX_NONE);
+
+					// If the constraint exists between two child bones, then no fix up is required
+					if (EditorSkelMesh->RefSkeleton.BoneIsChildOf(ExistingConstraintBoneIndex, NewBoneIndex))
+					{
+						continue;
+					}
+					
+					// If the constraint isn't between two child bones, then it is between a physical bone higher in the bone
+					// hierarchy than the new bone, so it needs to be fixed up by setting the constraint to point to the new bone
+					// instead. Additionally, collision needs to be re-enabled between the child bone and the identified "grandparent"
+					// bone.
+					const int32 ExistingConstraintBodyIndex = PhysicsAsset->FindBodyIndex(ExistingConstraintSetup->DefaultInstance.ConstraintBone2);
+					check(ExistingConstraintBodyIndex != INDEX_NONE);
+					check(ExistingConstraintBodyIndex == ParentBodyIndex);
+
+					SetCollisionBetween(ChildBodyIndex, ExistingConstraintBodyIndex, true);
+				}
+
+				UPhysicsConstraintTemplate* ChildConstraintSetup = PhysicsAsset->ConstraintSetup[ ConstraintIndex ];
+				check(ChildConstraintSetup);
+
+				InitConstraintSetup(ChildConstraintSetup, ChildBodyIndex, NewBodyIndex);
+			}
+		}
+	}
+
+	// If we have a physics parent, create a joint to it.
+	if (ParentBodyIndex != INDEX_NONE)
+	{
+		const int32 NewConstraintIndex = FPhysicsAssetUtils::CreateNewConstraint(PhysicsAsset, NewBoneName);
+		UPhysicsConstraintTemplate* ConstraintSetup = PhysicsAsset->ConstraintSetup[ NewConstraintIndex ];
+		check(ConstraintSetup);
+
+		InitConstraintSetup(ConstraintSetup, NewBodyIndex, ParentBodyIndex);
+	}
+
+	// update the tree
+	HierarchyChangedEvent.Broadcast();
+
+	if (bAutoSelect)
+	{
+		SetSelectedBodyAnyPrim(NewBodyIndex, true);
+	}
+	
+
+	RefreshPhysicsAssetChange(PhysicsAsset);
+}
+
+void FPhysicsAssetEditorSharedData::MakeNewConstraint(int32 BodyIndex0, int32 BodyIndex1)
+{
+	// check we have valid bodies
+	check(BodyIndex0 < PhysicsAsset->SkeletalBodySetups.Num());
+	check(BodyIndex1 < PhysicsAsset->SkeletalBodySetups.Num());
+
+	// Make a new unique name for this constraint
+	int32 Index = 0;
+	FString BaseConstraintName(TEXT("UserConstraint"));
+	FString ConstraintName = BaseConstraintName;
+	while(PhysicsAsset->FindConstraintIndex(*ConstraintName) != INDEX_NONE)
+	{
+		ConstraintName = FString::Printf(TEXT("%s_%d"), *BaseConstraintName, Index++);
+	}
+
+	// Create new constraint with a name not related to a bone, so it wont get auto managed in code that creates new bodies
+	const int32 NewConstraintIndex = FPhysicsAssetUtils::CreateNewConstraint(PhysicsAsset, *ConstraintName);
+	UPhysicsConstraintTemplate* ConstraintSetup = PhysicsAsset->ConstraintSetup[NewConstraintIndex];
+	check(ConstraintSetup);
+
+	InitConstraintSetup(ConstraintSetup, BodyIndex1, BodyIndex0);
+
+	// update the tree
+	HierarchyChangedEvent.Broadcast();
+	RefreshPhysicsAssetChange(PhysicsAsset);
+}
+
+void FPhysicsAssetEditorSharedData::SetConstraintRelTM(const FPhysicsAssetEditorSharedData::FSelection* Constraint, const FTransform& RelTM)
+{
+	USkeletalMesh* EditorSkelMesh = PhysicsAsset->GetPreviewMesh();
+	if(EditorSkelMesh == nullptr)
+	{
+		return;
+	}
+
+	FTransform WParentFrame = GetConstraintWorldTM(Constraint, EConstraintFrame::Frame2);
+	FTransform WNewChildFrame = RelTM * WParentFrame;
+
+	UPhysicsConstraintTemplate* ConstraintSetup = PhysicsAsset->ConstraintSetup[Constraint->Index];
+	ConstraintSetup->Modify();
+
+	// Get child bone transform
+	int32 BoneIndex = EditorSkelMesh->RefSkeleton.FindBoneIndex(ConstraintSetup->DefaultInstance.ConstraintBone1);
+	if (BoneIndex != INDEX_NONE)
+	{
+		FTransform BoneTM = EditorSkelComp->GetBoneTransform(BoneIndex);
+		BoneTM.RemoveScaling();
+
+		ConstraintSetup->DefaultInstance.SetRefFrame(EConstraintFrame::Frame1, WNewChildFrame.GetRelativeTransform(BoneTM));
+	}
+}
+
+void FPhysicsAssetEditorSharedData::SnapConstraintToBone(int32 ConstraintIndex)
+{
+	UPhysicsConstraintTemplate* ConstraintSetup = PhysicsAsset->ConstraintSetup[ConstraintIndex];
+	ConstraintSetup->Modify();
+	SnapConstraintToBone(ConstraintSetup->DefaultInstance);
+}
+
+void FPhysicsAssetEditorSharedData::SnapConstraintToBone(FConstraintInstance& ConstraintInstance)
+{
+	USkeletalMesh* EditorSkelMesh = PhysicsAsset->GetPreviewMesh();
+	if(EditorSkelMesh == nullptr)
+	{
+		return;
+	}
+
+	const int32 BoneIndex1 = EditorSkelMesh->RefSkeleton.FindBoneIndex(ConstraintInstance.ConstraintBone1);
+	const int32 BoneIndex2 = EditorSkelMesh->RefSkeleton.FindBoneIndex(ConstraintInstance.ConstraintBone2);
+
+	check(BoneIndex1 != INDEX_NONE);
+	check(BoneIndex2 != INDEX_NONE);
+
+	const FTransform BoneTransform1 = EditorSkelComp->GetBoneTransform(BoneIndex1);
+	const FTransform BoneTransform2 = EditorSkelComp->GetBoneTransform(BoneIndex2);
+
+	// Bone transforms are world space, and frame transforms are local space (local to bones).
+	// Frame 1 is the child frame, and set to identity.
+	// Frame 2 is the parent frame, and needs to be set relative to Frame1.
+	ConstraintInstance.SetRefFrame(EConstraintFrame::Frame2, BoneTransform1.GetRelativeTransform(BoneTransform2));
+	ConstraintInstance.SetRefFrame(EConstraintFrame::Frame1, FTransform::Identity);
+}
+
+void FPhysicsAssetEditorSharedData::CopyConstraint()
+{
+	check(SelectedConstraints.Num() == 1);
+
+	CopiedConstraintTemplate = PhysicsAsset->ConstraintSetup[GetSelectedConstraint()->Index];
+}
+
+void FPhysicsAssetEditorSharedData::PasteConstraintProperties()
+{
+	if (CopiedConstraintTemplate == NULL)
+	{
+		return;
+	}
+
+	const FScopedTransaction Transaction( NSLOCTEXT("PhysicsAssetEditor", "PasteConstraintProperties", "Paste Constraint Properties") );
+
+	UPhysicsConstraintTemplate* FromConstraintSetup = CopiedConstraintTemplate;
+
+	for(int32 i=0; i<SelectedConstraints.Num(); ++i)
+	{
+		// If we are showing instance properties - copy instance properties. If showing setup, just copy setup properties.
+		UPhysicsConstraintTemplate* ToConstraintSetup = PhysicsAsset->ConstraintSetup[SelectedConstraints[i].Index];
+		CopyConstraintProperties(FromConstraintSetup, ToConstraintSetup);
+	}
+	
+}
+
+void CycleMatrixRows(FMatrix* TM)
+{
+	float Tmp[3];
+
+	Tmp[0]		= TM->M[0][0];	Tmp[1]		= TM->M[0][1];	Tmp[2]		= TM->M[0][2];
+	TM->M[0][0] = TM->M[1][0];	TM->M[0][1] = TM->M[1][1];	TM->M[0][2] = TM->M[1][2];
+	TM->M[1][0] = TM->M[2][0];	TM->M[1][1] = TM->M[2][1];	TM->M[1][2] = TM->M[2][2];
+	TM->M[2][0] = Tmp[0];		TM->M[2][1] = Tmp[1];		TM->M[2][2] = Tmp[2];
+}
+
+void FPhysicsAssetEditorSharedData::CycleCurrentConstraintOrientation()
+{
+	UPhysicsConstraintTemplate* ConstraintTemplate = PhysicsAsset->ConstraintSetup[GetSelectedConstraint()->Index];
+	FMatrix ConstraintTransform = ConstraintTemplate->DefaultInstance.GetRefFrame(EConstraintFrame::Frame2).ToMatrixWithScale();
+	FTransform WParentFrame = GetConstraintWorldTM(GetSelectedConstraint(), EConstraintFrame::Frame2);
+	FTransform WChildFrame = GetConstraintWorldTM(GetSelectedConstraint(), EConstraintFrame::Frame1);
+	FTransform RelativeTransform = WChildFrame * WParentFrame.Inverse();
+
+	CycleMatrixRows(&ConstraintTransform);
+
+	ConstraintTemplate->DefaultInstance.SetRefFrame(EConstraintFrame::Frame2, FTransform(ConstraintTransform));
+	SetSelectedConstraintRelTM(RelativeTransform);
+}
+
+void FPhysicsAssetEditorSharedData::CycleCurrentConstraintActive()
+{
+	for(int32 i=0; i<SelectedConstraints.Num(); ++i)
+	{
+		UPhysicsConstraintTemplate* ConstraintTemplate = PhysicsAsset->ConstraintSetup[GetSelectedConstraint()->Index];
+		FConstraintInstance & DefaultInstance = ConstraintTemplate->DefaultInstance;
+
+		if(DefaultInstance.GetAngularSwing1Motion() != ACM_Limited && DefaultInstance.GetAngularSwing2Motion() != ACM_Limited)
+		{
+			DefaultInstance.SetAngularSwing1Motion(ACM_Limited);
+			DefaultInstance.SetAngularSwing2Motion(ACM_Locked);
+			DefaultInstance.SetAngularTwistMotion(ACM_Locked);
+		}else if(DefaultInstance.GetAngularSwing2Motion() != ACM_Limited && DefaultInstance.GetAngularTwistMotion() != ACM_Limited)
+		{
+			DefaultInstance.SetAngularSwing1Motion(ACM_Locked);
+			DefaultInstance.SetAngularSwing2Motion(ACM_Limited);
+			DefaultInstance.SetAngularTwistMotion(ACM_Locked);
+		}else
+		{
+			DefaultInstance.SetAngularSwing1Motion(ACM_Locked);
+			DefaultInstance.SetAngularSwing2Motion(ACM_Locked);
+			DefaultInstance.SetAngularTwistMotion(ACM_Limited);
+		}
+		
+	}
+}
+
+void FPhysicsAssetEditorSharedData::ToggleConstraint(EPhysicsAssetEditorConstraintType Constraint)
+{
+	for(int32 i=0; i<SelectedConstraints.Num(); ++i)
+	{
+		UPhysicsConstraintTemplate* ConstraintTemplate = PhysicsAsset->ConstraintSetup[GetSelectedConstraint()->Index];
+		FConstraintInstance & DefaultInstance = ConstraintTemplate->DefaultInstance;
+
+		if(Constraint == PCT_Swing1)
+		{
+			DefaultInstance.SetAngularSwing1Motion(DefaultInstance.GetAngularSwing1Motion() == ACM_Limited ? ACM_Locked : ACM_Limited);
+		}else if(Constraint == PCT_Swing2)
+		{
+			DefaultInstance.SetAngularSwing2Motion(DefaultInstance.GetAngularSwing2Motion() == ACM_Limited ? ACM_Locked : ACM_Limited);
+		}else
+		{
+			DefaultInstance.SetAngularTwistMotion(DefaultInstance.GetAngularTwistMotion() == ACM_Limited ? ACM_Locked : ACM_Limited);
+		}
+		
+	}
+}
+
+bool FPhysicsAssetEditorSharedData::IsAngularConstraintLocked(EPhysicsAssetEditorConstraintType Constraint) const
+{
+	bool bLocked = false;
+	bool bSame = false;
+
+	for(int32 i = 0; i < SelectedConstraints.Num(); ++i)
+	{
+		UPhysicsConstraintTemplate* ConstraintTemplate = PhysicsAsset->ConstraintSetup[GetSelectedConstraint()->Index];
+		FConstraintInstance & DefaultInstance = ConstraintTemplate->DefaultInstance;
+
+		if(Constraint == PCT_Swing1)
+		{
+			bLocked |= DefaultInstance.GetAngularSwing1Motion() == ACM_Locked;
+		}
+		else if(Constraint == PCT_Swing2)
+		{
+			bLocked |= DefaultInstance.GetAngularSwing2Motion() == ACM_Locked;
+		}
+		else
+		{
+			bLocked |= DefaultInstance.GetAngularTwistMotion() == ACM_Locked;
+		}
+	}
+
+	return bLocked;
+}
+
+
+void FPhysicsAssetEditorSharedData::DeleteBody(int32 DelBodyIndex, bool bRefreshComponent)
+{
+	USkeletalMesh* EditorSkelMesh = PhysicsAsset->GetPreviewMesh();
+	if(EditorSkelMesh == nullptr)
+	{
+		return;
+	}
+
+	const FScopedTransaction Transaction( NSLOCTEXT("UnrealEd", "DeleteBody", "Delete Body") );
+
+	// The physics asset and default instance..
+	PhysicsAsset->Modify();
+
+	// .. the body..
+	UBodySetup * BodySetup = PhysicsAsset->SkeletalBodySetups[DelBodyIndex];
+	BodySetup->Modify();	
+
+	// .. and any constraints to the body.
+	TArray<int32>	Constraints;
+	PhysicsAsset->BodyFindConstraints(DelBodyIndex, Constraints);
+
+	//we want to fixup constraints so that nearest child bodies get constraint with parent body
+	TArray<int32> NearestBodiesBelow;
+	PhysicsAsset->GetNearestBodyIndicesBelow(NearestBodiesBelow, BodySetup->BoneName, EditorSkelMesh);
+	
+	int32 BoneIndex = EditorSkelMesh->RefSkeleton.FindBoneIndex(BodySetup->BoneName);
+
+	if (BoneIndex != INDEX_NONE)	//it's possible to delete bodies that have no bones. In this case just ignore all of this fixup code
+	{
+		int32 ParentBodyIndex = PhysicsAsset->FindParentBodyIndex(EditorSkelMesh, BoneIndex);
+
+		UBodySetup * ParentBody = ParentBodyIndex != INDEX_NONE ? PhysicsAsset->SkeletalBodySetups[ParentBodyIndex] : NULL;
+
+		for (const int32 ConstraintIndex : Constraints)
+		{
+			UPhysicsConstraintTemplate * Constraint = PhysicsAsset->ConstraintSetup[ConstraintIndex];
+			Constraint->Modify();
+
+			if (ParentBody)
+			{
+				//for all constraints that contain a nearest child of this body, create a copy of the constraint between the child and parent
+				for (const int32 BodyBelowIndex : NearestBodiesBelow)
+				{
+					UBodySetup * BodyBelow = PhysicsAsset->SkeletalBodySetups[BodyBelowIndex];
+
+					if (Constraint->DefaultInstance.ConstraintBone1 == BodyBelow->BoneName)
+					{
+						int32 NewConstraintIndex = FPhysicsAssetUtils::CreateNewConstraint(PhysicsAsset, BodyBelow->BoneName, Constraint);
+						UPhysicsConstraintTemplate * NewConstraint = PhysicsAsset->ConstraintSetup[NewConstraintIndex];
+						InitConstraintSetup(NewConstraint, BodyBelowIndex, ParentBodyIndex);
+					}
+				}
+			}
+		}
+	}
+
+	// Now actually destroy body. This will destroy any constraints associated with the body as well.
+	FPhysicsAssetUtils::DestroyBody(PhysicsAsset, DelBodyIndex);
+
+	// Select nothing.
+	ClearSelectedBody();
+	ClearSelectedConstraints();
+	HierarchyChangedEvent.Broadcast();
+
+	if (bRefreshComponent)
+	{
+		RefreshPhysicsAssetChange(PhysicsAsset);
+	}
+}
+
+void FPhysicsAssetEditorSharedData::DeleteCurrentPrim()
+{
+	if (bRunningSimulation)
+	{
+		return;
+	}
+
+	if (!GetSelectedBody())
+	{
+		return;
+	}
+
+	// Make sure rendering is done - so we are not changing data being used by collision drawing.
+	FlushRenderingCommands();
+
+	//We will first get all the bodysetups we're interested in. The number of duplicates each bodysetup has tells us how many geoms are being deleted
+	//We need to do this first because deleting will modify our selection
+	TMap<UBodySetup *, TArray<FSelection>> BodySelectionMap;
+	TArray<UBodySetup*> BodySetups;
+	for(int32 i=0; i<SelectedBodies.Num(); ++i)
+	{
+		UBodySetup* BodySetup = PhysicsAsset->SkeletalBodySetups[SelectedBodies[i].Index];
+		BodySelectionMap.FindOrAdd(BodySetup).Add(SelectedBodies[i]);
+	}
+
+	const FScopedTransaction Transaction( NSLOCTEXT("UnrealEd", "DeletePrimitive", "Delete Primitive") );
+
+	for (TMap<UBodySetup*, TArray<FSelection> >::TConstIterator It(BodySelectionMap); It; ++It)
+	{
+		UBodySetup * BodySetup = It.Key();
+		const TArray<FSelection> & SelectedPrimitives = It.Value();
+
+		int32 SphereDeletedCount = 0;
+		int32 BoxDeletedCount = 0;
+		int32 SphylDeletedCount = 0;
+		int32 ConvexDeletedCount = 0;
+
+		for (int32 i = 0; i < SelectedPrimitives.Num(); ++i)
+		{
+			const FSelection& SelectedBody = SelectedPrimitives[i];
+			int32 BodyIndex = PhysicsAsset->FindBodyIndex(BodySetup->BoneName);
+
+			BodySetup->Modify();
+
+			if (SelectedBody.PrimitiveType == EAggCollisionShape::Sphere)
+			{
+				BodySetup->AggGeom.SphereElems.RemoveAt(SelectedBody.PrimitiveIndex - (SphereDeletedCount++));
+			}
+			else if (SelectedBody.PrimitiveType == EAggCollisionShape::Box)
+			{
+				BodySetup->AggGeom.BoxElems.RemoveAt(SelectedBody.PrimitiveIndex - (BoxDeletedCount++));
+			}
+			else if (SelectedBody.PrimitiveType == EAggCollisionShape::Sphyl)
+			{
+				BodySetup->AggGeom.SphylElems.RemoveAt(SelectedBody.PrimitiveIndex - (SphylDeletedCount++));
+			}
+			else if (SelectedBody.PrimitiveType == EAggCollisionShape::Convex)
+			{
+				BodySetup->AggGeom.ConvexElems.RemoveAt(SelectedBody.PrimitiveIndex - (ConvexDeletedCount++));
+				// Need to invalidate GUID in this case as cooked data must be updated
+				BodySetup->InvalidatePhysicsData();
+			}
+
+			// If this bone has no more geometry - remove it totally.
+			if (BodySetup->AggGeom.GetElementCount() == 0)
+			{
+				check(i == SelectedPrimitives.Num() - 1);	//we should really only delete on last prim - only reason this is even in for loop is because of API needing body index
+				if (BodyIndex != INDEX_NONE)
+				{
+					DeleteBody(BodyIndex, false);
+				}
+
+				if (CopiedBodySetup == BodySetup)
+				{
+					CopiedBodySetup = NULL;
+				}
+			}
+		}
+	}
+	
+	HierarchyChangedEvent.Broadcast();
+	ClearSelectedBody(); // Will call UpdateViewport
+	RefreshPhysicsAssetChange(PhysicsAsset);
+}
+FTransform FPhysicsAssetEditorSharedData::GetConstraintBodyTM(const UPhysicsConstraintTemplate* ConstraintSetup, EConstraintFrame::Type Frame) const
+{
+	if (ConstraintSetup == NULL)
+	{
+		return FTransform::Identity;
+	}
+
+	USkeletalMesh* EditorSkelMesh = PhysicsAsset->GetPreviewMesh();
+	if(EditorSkelMesh == nullptr)
+	{
+		return FTransform::Identity;
+	}
+
+	int32 BoneIndex;
+	if (Frame == EConstraintFrame::Frame1)
+	{
+		BoneIndex = EditorSkelMesh->RefSkeleton.FindBoneIndex(ConstraintSetup->DefaultInstance.ConstraintBone1);
+	}
+	else
+	{
+		BoneIndex = EditorSkelMesh->RefSkeleton.FindBoneIndex(ConstraintSetup->DefaultInstance.ConstraintBone2);
+	}
+
+	// If we couldn't find the bone - fall back to identity.
+	if (BoneIndex == INDEX_NONE)
+	{
+		return FTransform::Identity;
+	}
+	else
+	{
+		FTransform BoneTM = EditorSkelComp->GetBoneTransform(BoneIndex);
+		BoneTM.RemoveScaling();
+
+		return BoneTM;
+	}
+}
+
+FTransform FPhysicsAssetEditorSharedData::GetConstraintWorldTM(const UPhysicsConstraintTemplate* ConstraintSetup, EConstraintFrame::Type Frame, float Scale) const
+{
+	if (ConstraintSetup == NULL)
+	{
+		return FTransform::Identity;
+	}
+
+	USkeletalMesh* EditorSkelMesh = PhysicsAsset->GetPreviewMesh();
+	if(EditorSkelMesh == nullptr)
+	{
+		return FTransform::Identity;
+	}
+
+	FVector Scale3D(Scale);
+
+	int32 BoneIndex;
+	FTransform LFrame = ConstraintSetup->DefaultInstance.GetRefFrame(Frame);
+	if (Frame == EConstraintFrame::Frame1)
+	{
+		BoneIndex = EditorSkelMesh->RefSkeleton.FindBoneIndex(ConstraintSetup->DefaultInstance.ConstraintBone1);
+	}
+	else
+	{
+		BoneIndex = EditorSkelMesh->RefSkeleton.FindBoneIndex(ConstraintSetup->DefaultInstance.ConstraintBone2);
+	}
+
+	// If we couldn't find the bone - fall back to identity.
+	if (BoneIndex == INDEX_NONE)
+	{
+		return FTransform::Identity;
+	}
+	else
+	{
+		FTransform BoneTM = EditorSkelComp->GetBoneTransform(BoneIndex);
+		BoneTM.RemoveScaling();
+
+		LFrame.ScaleTranslation(Scale3D);
+
+		return LFrame * BoneTM;
+	}
+}
+
+FTransform FPhysicsAssetEditorSharedData::GetConstraintMatrix(int32 ConstraintIndex, EConstraintFrame::Type Frame, float Scale) const
+{
+	UPhysicsConstraintTemplate* ConstraintSetup = PhysicsAsset->ConstraintSetup[ConstraintIndex];
+	return GetConstraintWorldTM(ConstraintSetup, Frame, Scale);
+}
+
+
+FTransform FPhysicsAssetEditorSharedData::GetConstraintWorldTM(const FSelection* Constraint, EConstraintFrame::Type Frame) const
+{
+	int32 ConstraintIndex = Constraint ? Constraint->Index : INDEX_NONE;
+	if (ConstraintIndex == INDEX_NONE)
+	{
+		return FTransform::Identity;
+	}
+
+	UPhysicsConstraintTemplate* ConstraintSetup = PhysicsAsset->ConstraintSetup[ConstraintIndex];
+	return GetConstraintWorldTM(ConstraintSetup, Frame, 1.f);
+}
+
+
+void FPhysicsAssetEditorSharedData::DeleteCurrentConstraint()
+{
+	if (!GetSelectedConstraint())
+	{
+		return;
+	}
+
+	const FScopedTransaction Transaction( NSLOCTEXT("PhysicsAssetEditor", "DeleteConstraint", "Delete Constraint") );
+
+	//Save indices before delete because delete modifies our Selected array
+	TArray<int32> Indices;
+	for(int32 i=0; i<SelectedConstraints.Num(); ++i)
+	{
+		Indices.Add(SelectedConstraints[i].Index);
+	}
+
+	Indices.Sort();
+
+	//These are indices into an array, we must remove it from greatest to smallest so that the indices don't shift
+	for(int32 i=Indices.Num() - 1; i>= 0; --i)
+	{
+		
+		if(PhysicsAsset->ConstraintSetup[Indices[i]] == CopiedConstraintTemplate)
+		{
+			CopiedConstraintTemplate = NULL;
+		}
+
+		PhysicsAsset->Modify();
+		FPhysicsAssetUtils::DestroyConstraint(PhysicsAsset, Indices[i]);
+		
+	}
+	
+	ClearSelectedConstraints();
+
+	EditorSkelComp->RecreatePhysicsState();
+	HierarchyChangedEvent.Broadcast();
+	PreviewChangedEvent.Broadcast();
+}
+
+void FPhysicsAssetEditorSharedData::ToggleSimulation()
+{
+	// don't start simulation if there are no bodies or if we are manipulating a body
+	if (PhysicsAsset->SkeletalBodySetups.Num() == 0 || bManipulating)
+	{
+		return;  
+	}
+
+	EnableSimulation(!bRunningSimulation);
+
+	bRunningSimulation = !bRunningSimulation;
+}
+
+void FPhysicsAssetEditorSharedData::EnableSimulation(bool bEnableSimulation)
+{
+	if (bEnableSimulation)
+	{
+		// We should not already have an instance (destroyed when stopping sim).
+		EditorSkelComp->SetSimulatePhysics(true);
+		EditorSkelComp->ResetAllBodiesSimulatePhysics();
+		EditorSkelComp->SetPhysicsBlendWeight(EditorOptions->PhysicsBlend);
+		PhysicalAnimationComponent->SetSkeletalMeshComponent(EditorSkelComp);
+
+		// Make it start simulating
+		EditorSkelComp->WakeAllRigidBodies();
+	}
+	else
+	{
+		// Stop any animation and clear node when stopping simulation.
+		PhysicalAnimationComponent->SetSkeletalMeshComponent(nullptr);
+
+		EditorSkelComp->SetPhysicsBlendWeight(0.f);
+		EditorSkelComp->ResetAllBodiesSimulatePhysics();
+		EditorSkelComp->SetSimulatePhysics(false);
+		ForceDisableSimulation();
+
+		// Since simulation, actor location changes. Reset to identity 
+		EditorSkelComp->SetWorldTransform(ResetTM);
+		// Force an update of the skeletal mesh to get it back to ref pose
+		EditorSkelComp->RefreshBoneTransforms();
+		
+		PreviewChangedEvent.Broadcast();
+	}
+}
+
+void FPhysicsAssetEditorSharedData::OpenNewBodyDlg()
+{
+	OpenNewBodyDlg(&NewBodyResponse);
+}
+
+void FPhysicsAssetEditorSharedData::OpenNewBodyDlg(EAppReturnType::Type* NewBodyResponse)
+{
+	TSharedRef<SWindow> ModalWindow = SNew(SWindow)
+		.Title(LOCTEXT("NewAssetTitle", "New Physics Asset"))
+		.SizingRule(ESizingRule::FixedSize)
+		.ClientSize(FVector2D(400.0f, 400.0f))
+		.SupportsMinimize(false) 
+		.SupportsMaximize(false);
+
+	TWeakPtr<SWindow> ModalWindowPtr = ModalWindow;
+
+	ModalWindow->SetContent(
+		CreateGenerateBodiesWidget(
+			FSimpleDelegate::CreateLambda([ModalWindowPtr, NewBodyResponse]()
+			{
+				*NewBodyResponse = EAppReturnType::Ok;
+				ModalWindowPtr.Pin()->RequestDestroyWindow();
+			}),
+			FSimpleDelegate::CreateLambda([ModalWindowPtr, NewBodyResponse]()
+			{
+				*NewBodyResponse = EAppReturnType::Cancel;
+				ModalWindowPtr.Pin()->RequestDestroyWindow();
+			}),
+			true,
+			LOCTEXT("CreateAsset", "Create Asset"),
+			true
+		));
+
+	GEditor->EditorAddModalWindow(ModalWindow);
+}
+
+TSharedRef<SWidget> FPhysicsAssetEditorSharedData::CreateGenerateBodiesWidget(const FSimpleDelegate& InOnCreate, const FSimpleDelegate& InOnCancel, const TAttribute<bool>& InIsEnabled, const TAttribute<FText>& InCreateButtonText, bool bForNewAsset)
+{
+	FDetailsViewArgs DetailsViewArgs(false, false, true, FDetailsViewArgs::HideNameArea, true);
+	DetailsViewArgs.bAllowSearch = false;
+
+	FPropertyEditorModule& PropertyEditorModule = FModuleManager::LoadModuleChecked<FPropertyEditorModule>("PropertyEditor");
+	TSharedRef<IDetailsView> DetailsView = PropertyEditorModule.CreateDetailView(DetailsViewArgs);
+
+	GetMutableDefault<UPhysicsAssetGenerationSettings>()->LoadConfig();
+	DetailsView->SetObject(GetMutableDefault<UPhysicsAssetGenerationSettings>());
+	DetailsView->OnFinishedChangingProperties().AddLambda([](const FPropertyChangedEvent& InEvent){ GetMutableDefault<UPhysicsAssetGenerationSettings>()->SaveConfig(); });
+
+	return SNew(SVerticalBox)
+		.IsEnabled(InIsEnabled)
+		+SVerticalBox::Slot()
+		.FillHeight(1.0f)
+		[
+			DetailsView
+		]
+		+SVerticalBox::Slot()
+		.AutoHeight()
+		[
+			SNew(SBorder)
+			.BorderImage(FEditorStyle::GetBrush("ToolPanel.GroupBorder"))
+			.VAlign(VAlign_Center)
+			.HAlign(HAlign_Right)
+			[
+				SNew(SHorizontalBox)
+				+SHorizontalBox::Slot()
+				.Padding(2.0f)
+				.AutoWidth()
+				[
+					SNew(SButton)
+					.ButtonStyle(FEditorStyle::Get(), "FlatButton.Success")
+					.ForegroundColor(FLinearColor::White)
+					.ContentPadding(FMargin(6, 2))
+					.OnClicked_Lambda([InOnCreate]()
+					{ 
+						GetMutableDefault<UPhysicsAssetGenerationSettings>()->SaveConfig(); 
+						InOnCreate.ExecuteIfBound(); 
+						return FReply::Handled(); 
+					})
+					.ToolTipText(bForNewAsset ? 
+								LOCTEXT("CreateAsset_Tooltip", "Create a new physics asset using these settings.") :
+								LOCTEXT("GenerateBodies_Tooltip", "Generate new bodies and constraints. If bodies are selected then they will be replaced along with their constraints using the new settings, otherwise all bodies and constraints will be re-created"))
+					[
+						SNew(STextBlock)
+						.TextStyle(FEditorStyle::Get(), "PhysicsAssetEditor.Tools.Font")
+						.Text(InCreateButtonText)
+					]
+				]
+				+SHorizontalBox::Slot()
+				.Padding(2.0f)
+				.AutoWidth()
+				[
+					SNew(SButton)
+					.Visibility_Lambda([bForNewAsset](){ return bForNewAsset ? EVisibility::Visible : EVisibility::Collapsed; })
+					.ButtonStyle(FEditorStyle::Get(), "FlatButton")
+					.ForegroundColor(FLinearColor::White)
+					.ContentPadding(FMargin(6, 2))
+					.OnClicked_Lambda([InOnCancel](){ InOnCancel.ExecuteIfBound(); return FReply::Handled(); })
+					[
+						SNew(STextBlock)
+						.TextStyle(FEditorStyle::Get(), "PhysicsAssetEditor.Tools.Font")
+						.Text(LOCTEXT("Cancel", "Cancel"))
+					]
+				]
+			]
+		];
+}
+
+void FPhysicsAssetEditorSharedData::PostUndo()
+{
+	if (bRunningSimulation)
+	{
+		return;
+	}
+
+	bool bInvalidSelection = false;
+
+	for (int32 BodyIndex = 0; BodyIndex < SelectedBodies.Num() && bInvalidSelection == false; ++BodyIndex)
+	{
+		const FSelection& Selection = SelectedBodies[BodyIndex];
+		if (PhysicsAsset->SkeletalBodySetups.Num() <= Selection.Index)
+		{
+			bInvalidSelection = true;
+		}
+		else
+		{
+			
+			if (UBodySetup * BodySetup = PhysicsAsset->SkeletalBodySetups[Selection.Index])
+			{
+				switch (Selection.PrimitiveType)
+				{
+				case EAggCollisionShape::Box: bInvalidSelection = BodySetup->AggGeom.BoxElems.Num() <= Selection.PrimitiveIndex ? true : bInvalidSelection; break;
+				case EAggCollisionShape::Convex: bInvalidSelection = BodySetup->AggGeom.ConvexElems.Num() <= Selection.PrimitiveIndex ? true : bInvalidSelection; break;
+				case EAggCollisionShape::Sphere: bInvalidSelection = BodySetup->AggGeom.SphereElems.Num() <= Selection.PrimitiveIndex ? true : bInvalidSelection; break;
+				case EAggCollisionShape::Sphyl: bInvalidSelection = BodySetup->AggGeom.SphylElems.Num() <= Selection.PrimitiveIndex ? true : bInvalidSelection; break;
+				default: bInvalidSelection = true;
+				}
+			}
+			else
+			{
+				bInvalidSelection = true;
+			}
+		}
+	}
+
+	for (int32 ConstraintIndex = 0; ConstraintIndex < SelectedConstraints.Num() && bInvalidSelection == false; ++ConstraintIndex)
+	{
+		const FSelection& Selection = SelectedConstraints[ConstraintIndex];
+		if (PhysicsAsset->ConstraintSetup.Num() <= Selection.Index)
+		{
+			bInvalidSelection = true;
+		}
+	}
+
+	if (bInvalidSelection)
+	{
+		// Clear selection before we undo. We don't transact the editor itself - don't want to have something selected that is then removed.
+		ClearSelectedBody();
+		ClearSelectedConstraints();
+	}
+
+	PreviewChangedEvent.Broadcast();
+	HierarchyChangedEvent.Broadcast();
+}
+
+void FPhysicsAssetEditorSharedData::Redo()
+{
+	if (bRunningSimulation)
+	{
+		return;
+	}
+
+	ClearSelectedBody();
+	ClearSelectedConstraints();
+
+	GEditor->RedoTransaction();
+	PhysicsAsset->UpdateBodySetupIndexMap();
+
+	PreviewChangedEvent.Broadcast();
+	HierarchyChangedEvent.Broadcast();
+}
+
+void FPhysicsAssetEditorSharedData::AddReferencedObjects(FReferenceCollector& Collector)
+{
+	Collector.AddReferencedObject(PhysicsAsset);
+	Collector.AddReferencedObject(EditorSkelComp);
+	Collector.AddReferencedObject(PhysicalAnimationComponent);
+	Collector.AddReferencedObject(EditorOptions);
+	Collector.AddReferencedObject(MouseHandle);
+	Collector.AddReferencedObject(CopiedBodySetup);
+	Collector.AddReferencedObject(CopiedConstraintTemplate);
+
+	PreviewScene.Pin()->AddReferencedObjects(Collector);
+}
+
+void FPhysicsAssetEditorSharedData::ForceDisableSimulation()
+{
+	// Reset simulation state of body instances so we dont actually simulate outside of 'simulation mode'
+	for (int32 BodyIdx = 0; BodyIdx < EditorSkelComp->Bodies.Num(); ++BodyIdx)
+	{
+		if (FBodyInstance* BodyInst = EditorSkelComp->Bodies[BodyIdx])
+		{
+			if (UBodySetup * PhysAssetBodySetup = PhysicsAsset->SkeletalBodySetups[BodyIdx])
+			{
+				BodyInst->SetInstanceSimulatePhysics(false);
+			}
+		}
+	}
+}
+
+#undef LOCTEXT_NAMESPACE
diff --git a/Engine/Source/Editor/PhysicsAssetEditor/Private/PhysicsAssetEditorSkeletalMeshComponent.cpp b/Engine/Source/Editor/PhysicsAssetEditor/Private/PhysicsAssetEditorSkeletalMeshComponent.cpp
new file mode 100644
--- /dev/null
+++ b/Engine/Source/Editor/PhysicsAssetEditor/Private/PhysicsAssetEditorSkeletalMeshComponent.cpp
@@ -0,0 +1,443 @@
+// Copyright 1998-2017 Epic Games, Inc. All Rights Reserved.
+
+#include "PhysicsAssetEditorSkeletalMeshComponent.h"
+#include "Materials/MaterialInterface.h"
+#include "Materials/MaterialInstanceDynamic.h"
+#include "Components/SkeletalMeshComponent.h"
+#include "Materials/Material.h"
+#include "Preferences/PhysicsAssetEditorOptions.h"
+#include "SceneManagement.h"
+#include "PhysicsAssetEditorSharedData.h"
+#include "PhysicsAssetEditorHitProxies.h"
+#include "PhysicsAssetEditorSkeletalMeshComponent.h"
+#include "PhysicsEngine/PhysicsConstraintTemplate.h"
+#include "PhysicsEngine/PhysicsAsset.h"
+#include "SkeletalMeshTypes.h"
+#include "UObject/Package.h"
+#include "EditorStyleSet.h"
+
+namespace
+{
+	// How large to make the constraint arrows.
+	// The factor of 60 was found experimentally, to look reasonable in comparison with the rest of the constraint visuals.
+	constexpr float ConstraintArrowScale = 60.0f;
+}
+
+UPhysicsAssetEditorSkeletalMeshComponent::UPhysicsAssetEditorSkeletalMeshComponent(const FObjectInitializer& ObjectInitializer)
+	: Super(ObjectInitializer)
+	, BoneUnselectedColor(170,155,225)
+	, NoCollisionColor(200, 200, 200)
+	, FixedColor(125,125,0)
+	, ConstraintBone1Color(255,166,0)
+	, ConstraintBone2Color(0,150,150)
+	, HierarchyDrawColor(220, 255, 220)
+	, AnimSkelDrawColor(255, 64, 64)
+	, COMRenderSize(5.0f)
+	, InfluenceLineLength(2.0f)
+	, InfluenceLineColor(0,255,0)
+{
+
+	// Body materials
+	UMaterialInterface* BaseElemSelectedMaterial = LoadObject<UMaterialInterface>(NULL, TEXT("/Engine/EditorMaterials/PhAT_ElemSelectedMaterial.PhAT_ElemSelectedMaterial"), NULL, LOAD_None, NULL);
+	ElemSelectedMaterial = UMaterialInstanceDynamic::Create(BaseElemSelectedMaterial, GetTransientPackage());
+	check(ElemSelectedMaterial);
+
+	UMaterialInterface* BaseBoneSelectedMaterial = LoadObject<UMaterialInterface>(NULL, TEXT("/Engine/EditorMaterials/PhAT_BoneSelectedMaterial.PhAT_BoneSelectedMaterial"), NULL, LOAD_None, NULL);
+	BoneSelectedMaterial = UMaterialInstanceDynamic::Create(BaseBoneSelectedMaterial, GetTransientPackage());
+	check(BoneSelectedMaterial);
+
+	BoneMaterialHit = UMaterial::GetDefaultMaterial(MD_Surface);
+	check(BoneMaterialHit);
+
+	UMaterialInterface* BaseBoneUnselectedMaterial = LoadObject<UMaterialInterface>(NULL, TEXT("/Engine/EditorMaterials/PhAT_UnselectedMaterial.PhAT_UnselectedMaterial"), NULL, LOAD_None, NULL);
+	BoneUnselectedMaterial = UMaterialInstanceDynamic::Create(BaseBoneUnselectedMaterial, GetTransientPackage());
+	check(BoneUnselectedMaterial);
+
+	UMaterialInterface* BaseBoneNoCollisionMaterial = LoadObject<UMaterialInterface>(NULL, TEXT("/Engine/EditorMaterials/PhAT_NoCollisionMaterial.PhAT_NoCollisionMaterial"), NULL, LOAD_None, NULL);
+	BoneNoCollisionMaterial = UMaterialInstanceDynamic::Create(BaseBoneNoCollisionMaterial, GetTransientPackage());
+	check(BoneNoCollisionMaterial);
+
+	// this is because in phat editor, you'd like to see fixed bones to be fixed without animation force update
+	KinematicBonesUpdateType = EKinematicBonesUpdateToPhysics::SkipSimulatingBones;
+	bUpdateJointsFromAnimation = false;
+	ForcedLodModel = 1;
+
+	static FName CollisionProfileName(TEXT("PhysicsActor"));
+	SetCollisionProfileName(CollisionProfileName);
+
+	bSelectable = false;
+}
+
+void UPhysicsAssetEditorSkeletalMeshComponent::RenderAssetTools(const FSceneView* View, class FPrimitiveDrawInterface* PDI)
+{
+	check(SharedData);
+
+	UPhysicsAsset* const PhysicsAsset = GetPhysicsAsset();
+	
+	if(!PhysicsAsset)
+	{
+		// Nothing to draw without an asset, this can happen if the preview scene has no skeletal mesh
+		return;
+	}
+
+	EPhysicsAssetEditorRenderMode CollisionViewMode = SharedData->GetCurrentCollisionViewMode(SharedData->bRunningSimulation);
+
+#if DEBUG_CLICK_VIEWPORT
+	PDI->DrawLine(SharedData->LastClickOrigin, SharedData->LastClickOrigin + SharedData->LastClickDirection * 5000.0f, FLinearColor(1, 1, 0, 1), SDPG_Foreground);
+	PDI->DrawPoint(SharedData->LastClickOrigin, FLinearColor(1, 0, 0), 5, SDPG_Foreground);
+#endif
+
+	// set opacity of our materials
+	static FName OpacityName(TEXT("Opacity"));
+	ElemSelectedMaterial->SetScalarParameterValue(OpacityName, SharedData->EditorOptions->CollisionOpacity);
+	BoneSelectedMaterial->SetScalarParameterValue(OpacityName, SharedData->EditorOptions->CollisionOpacity);
+	BoneUnselectedMaterial->SetScalarParameterValue(OpacityName, SharedData->EditorOptions->bSolidRenderingForSelectedOnly ? 0.0f : SharedData->EditorOptions->CollisionOpacity);
+	BoneNoCollisionMaterial->SetScalarParameterValue(OpacityName, SharedData->EditorOptions->bSolidRenderingForSelectedOnly ? 0.0f : SharedData->EditorOptions->CollisionOpacity);
+
+	static FName SelectionColorName(TEXT("SelectionColor"));
+	const FSlateColor SelectionColor = FEditorStyle::GetSlateColor(SelectionColorName);
+	const FLinearColor LinearSelectionColor( SelectionColor.IsColorSpecified() ? SelectionColor.GetSpecifiedColor() : FLinearColor::White );
+
+	ElemSelectedMaterial->SetVectorParameterValue(SelectionColorName, LinearSelectionColor);
+	BoneSelectedMaterial->SetVectorParameterValue(SelectionColorName, LinearSelectionColor);
+
+	// Draw bodies
+	for (int32 i = 0; i <PhysicsAsset->SkeletalBodySetups.Num(); ++i)
+	{
+		int32 BoneIndex = GetBoneIndex(PhysicsAsset->SkeletalBodySetups[i]->BoneName);
+
+		// If we found a bone for it, draw the collision.
+		// The logic is as follows; always render in the ViewMode requested when not in hit mode - but if we are in hit mode and the right editing mode, render as solid
+		if (BoneIndex != INDEX_NONE)
+		{
+			FTransform BoneTM = GetBoneTransform(BoneIndex);
+			float Scale = BoneTM.GetScale3D().GetAbsMax();
+			FVector VectorScale(Scale);
+			BoneTM.RemoveScaling();
+
+			FKAggregateGeom* AggGeom = &PhysicsAsset->SkeletalBodySetups[i]->AggGeom;
+
+			for (int32 j = 0; j <AggGeom->SphereElems.Num(); ++j)
+			{
+				PDI->SetHitProxy(new HPhysicsAssetEditorEdBoneProxy(i, EAggCollisionShape::Sphere, j));
+
+				FTransform ElemTM = GetPrimitiveTransform(BoneTM, i, EAggCollisionShape::Sphere, j, Scale);
+
+
+				//solids are drawn if it's the ViewMode and we're not doing a hit, or if it's hitAndBodyMode
+				if(CollisionViewMode == EPhysicsAssetEditorRenderMode::Solid)
+				{
+					UMaterialInterface*	PrimMaterial = GetPrimitiveMaterial(i, EAggCollisionShape::Sphere, j);
+					AggGeom->SphereElems[j].DrawElemSolid(PDI, ElemTM, VectorScale, PrimMaterial->GetRenderProxy(0));
+				}
+
+				if (CollisionViewMode == EPhysicsAssetEditorRenderMode::Solid || CollisionViewMode == EPhysicsAssetEditorRenderMode::Wireframe)
+				{
+					AggGeom->SphereElems[j].DrawElemWire(PDI, ElemTM, VectorScale, GetPrimitiveColor(i, EAggCollisionShape::Sphere, j));
+				}
+
+				PDI->SetHitProxy(NULL);
+			}
+
+			for (int32 j = 0; j <AggGeom->BoxElems.Num(); ++j)
+			{
+				PDI->SetHitProxy(new HPhysicsAssetEditorEdBoneProxy(i, EAggCollisionShape::Box, j));
+
+				FTransform ElemTM = GetPrimitiveTransform(BoneTM, i, EAggCollisionShape::Box, j, Scale);
+
+				if (CollisionViewMode == EPhysicsAssetEditorRenderMode::Solid)
+				{
+					UMaterialInterface*	PrimMaterial = GetPrimitiveMaterial(i, EAggCollisionShape::Box, j);
+					AggGeom->BoxElems[j].DrawElemSolid(PDI, ElemTM, VectorScale, PrimMaterial->GetRenderProxy(0));
+				}
+
+				if (CollisionViewMode == EPhysicsAssetEditorRenderMode::Solid || CollisionViewMode == EPhysicsAssetEditorRenderMode::Wireframe)
+				{
+					AggGeom->BoxElems[j].DrawElemWire(PDI, ElemTM, VectorScale, GetPrimitiveColor(i, EAggCollisionShape::Box, j));
+				}
+				
+				PDI->SetHitProxy(NULL);
+			}
+
+			for (int32 j = 0; j <AggGeom->SphylElems.Num(); ++j)
+			{
+				PDI->SetHitProxy(new HPhysicsAssetEditorEdBoneProxy(i, EAggCollisionShape::Sphyl, j));
+
+				FTransform ElemTM = GetPrimitiveTransform(BoneTM, i, EAggCollisionShape::Sphyl, j, Scale);
+
+				if (CollisionViewMode == EPhysicsAssetEditorRenderMode::Solid)
+				{
+					UMaterialInterface*	PrimMaterial = GetPrimitiveMaterial(i, EAggCollisionShape::Sphyl, j);
+					AggGeom->SphylElems[j].DrawElemSolid(PDI, ElemTM, VectorScale, PrimMaterial->GetRenderProxy(0));
+				}
+
+				if (CollisionViewMode == EPhysicsAssetEditorRenderMode::Solid || CollisionViewMode == EPhysicsAssetEditorRenderMode::Wireframe)
+				{
+					AggGeom->SphylElems[j].DrawElemWire(PDI, ElemTM, VectorScale, GetPrimitiveColor(i, EAggCollisionShape::Sphyl, j));
+				}
+					
+				PDI->SetHitProxy(NULL);
+			}
+
+			for (int32 j = 0; j <AggGeom->ConvexElems.Num(); ++j)
+			{
+				PDI->SetHitProxy(new HPhysicsAssetEditorEdBoneProxy(i, EAggCollisionShape::Convex, j));
+
+				FTransform ElemTM = GetPrimitiveTransform(BoneTM, i, EAggCollisionShape::Convex, j, Scale);
+
+				//convex doesn't have solid draw so render lines if we're in hitTestAndBodyMode
+				if (CollisionViewMode == EPhysicsAssetEditorRenderMode::Solid || CollisionViewMode == EPhysicsAssetEditorRenderMode::Wireframe)
+				{
+					AggGeom->ConvexElems[j].DrawElemWire(PDI, ElemTM, Scale, GetPrimitiveColor(i, EAggCollisionShape::Convex, j));
+				}
+				
+				PDI->SetHitProxy(NULL);
+			}
+
+			if (SharedData->bShowCOM && Bodies.IsValidIndex(i))
+			{
+				Bodies[i]->DrawCOMPosition(PDI, COMRenderSize, SharedData->COMRenderColor);
+			}
+		}
+	}
+
+	// Draw Constraints
+	EPhysicsAssetEditorConstraintViewMode ConstraintViewMode = SharedData->GetCurrentConstraintViewMode(SharedData->bRunningSimulation);
+	if (ConstraintViewMode != EPhysicsAssetEditorConstraintViewMode::None)
+	{
+		for (int32 i = 0; i <PhysicsAsset->ConstraintSetup.Num(); ++i)
+		{
+			int32 BoneIndex1 = GetBoneIndex(PhysicsAsset->ConstraintSetup[i]->DefaultInstance.ConstraintBone1);
+			int32 BoneIndex2 = GetBoneIndex(PhysicsAsset->ConstraintSetup[i]->DefaultInstance.ConstraintBone2);
+			// if bone doesn't exist, do not draw it. It crashes in random points when we try to manipulate. 
+			if (BoneIndex1 != INDEX_NONE && BoneIndex2 != INDEX_NONE)
+			{
+				PDI->SetHitProxy(new HPhysicsAssetEditorEdConstraintProxy(i));
+
+				DrawConstraint(i, View, PDI, SharedData->EditorOptions->bShowConstraintsAsPoints);
+
+				PDI->SetHitProxy(NULL);
+			}
+		}
+	}
+}
+
+void UPhysicsAssetEditorSkeletalMeshComponent::DebugDraw(const FSceneView* View, FPrimitiveDrawInterface* PDI)
+{
+	RenderAssetTools(View, PDI);
+}
+
+FPrimitiveSceneProxy* UPhysicsAssetEditorSkeletalMeshComponent::CreateSceneProxy()
+{
+	FPrimitiveSceneProxy* Proxy = NULL;
+	EPhysicsAssetEditorRenderMode MeshViewMode = SharedData->GetCurrentMeshViewMode(SharedData->bRunningSimulation);
+	if (MeshViewMode != EPhysicsAssetEditorRenderMode::None)
+	{
+		Proxy = UDebugSkelMeshComponent::CreateSceneProxy();
+	}
+
+	return Proxy;
+}
+
+bool ConstraintInSelected(int32 Index, const TArray<FPhysicsAssetEditorSharedData::FSelection> & Constraints)
+{
+	for(int32 i=0; i<Constraints.Num(); ++i)
+	{
+
+		if(Constraints[i].Index == Index)
+		{
+			return true;
+		}
+	}
+
+	return false;
+}
+
+void UPhysicsAssetEditorSkeletalMeshComponent::DrawConstraint(int32 ConstraintIndex, const FSceneView* View, FPrimitiveDrawInterface* PDI, bool bDrawAsPoint)
+{
+	EPhysicsAssetEditorConstraintViewMode ConstraintViewMode = SharedData->GetCurrentConstraintViewMode(SharedData->bRunningSimulation);
+	bool bDrawLimits = false;
+	bool bConstraintSelected = ConstraintInSelected(ConstraintIndex, SharedData->SelectedConstraints);
+	if (ConstraintViewMode == EPhysicsAssetEditorConstraintViewMode::AllLimits || bConstraintSelected)
+	{
+		bDrawLimits = true;
+	}
+
+	bool bDrawSelected = false;
+	if (!SharedData->bRunningSimulation && bConstraintSelected)
+	{
+		bDrawSelected = true;
+	}
+
+	UPhysicsConstraintTemplate* ConstraintSetup = SharedData->PhysicsAsset->ConstraintSetup[ConstraintIndex];
+
+	FTransform Con1Frame = SharedData->GetConstraintMatrix(ConstraintIndex, EConstraintFrame::Frame1, 1.f);
+	FTransform Con2Frame = SharedData->GetConstraintMatrix(ConstraintIndex, EConstraintFrame::Frame2, 1.f);
+
+	const float DrawScale = ConstraintArrowScale * SharedData->EditorOptions->ConstraintDrawSize;
+
+	ConstraintSetup->DefaultInstance.DrawConstraint(PDI, SharedData->EditorOptions->ConstraintDrawSize, DrawScale, bDrawLimits, bDrawSelected, Con1Frame, Con2Frame, bDrawAsPoint);
+}
+
+FTransform UPhysicsAssetEditorSkeletalMeshComponent::GetPrimitiveTransform(FTransform& BoneTM, int32 BodyIndex, EAggCollisionShape::Type PrimType, int32 PrimIndex, float Scale)
+{
+	UBodySetup* SharedBodySetup = SharedData->PhysicsAsset->SkeletalBodySetups[BodyIndex];
+	FVector Scale3D(Scale);
+
+	FTransform ManTM = FTransform::Identity;
+
+	if(SharedData->bManipulating && !SharedData->bRunningSimulation)
+	{
+		FPhysicsAssetEditorSharedData::FSelection Body(BodyIndex, PrimType, PrimIndex);
+		for(int32 i=0; i<SharedData->SelectedBodies.Num(); ++i)
+		{
+			if (Body == SharedData->SelectedBodies[i])
+			{
+				ManTM = SharedData->SelectedBodies[i].ManipulateTM;
+				break;
+			}
+
+		}
+	}
+
+
+	if (PrimType == EAggCollisionShape::Sphere)
+	{
+		FTransform PrimTM = ManTM * SharedBodySetup->AggGeom.SphereElems[PrimIndex].GetTransform();
+		PrimTM.ScaleTranslation(Scale3D);
+		return PrimTM * BoneTM;
+	}
+	else if (PrimType == EAggCollisionShape::Box)
+	{
+		FTransform PrimTM = ManTM * SharedBodySetup->AggGeom.BoxElems[PrimIndex].GetTransform();
+		PrimTM.ScaleTranslation(Scale3D);
+		return PrimTM * BoneTM;
+	}
+	else if (PrimType == EAggCollisionShape::Sphyl)
+	{
+		FTransform PrimTM = ManTM * SharedBodySetup->AggGeom.SphylElems[PrimIndex].GetTransform();
+		PrimTM.ScaleTranslation(Scale3D);
+		return PrimTM * BoneTM;
+	}
+	else if (PrimType == EAggCollisionShape::Convex)
+	{
+		FTransform PrimTM = ManTM * SharedBodySetup->AggGeom.ConvexElems[PrimIndex].GetTransform();
+		PrimTM.ScaleTranslation(Scale3D);
+		return PrimTM * BoneTM;
+	}
+
+	// Should never reach here
+	check(0);
+	return FTransform::Identity;
+}
+
+FColor UPhysicsAssetEditorSkeletalMeshComponent::GetPrimitiveColor(int32 BodyIndex, EAggCollisionShape::Type PrimitiveType, int32 PrimitiveIndex)
+{
+	UBodySetup* SharedBodySetup = SharedData->PhysicsAsset->SkeletalBodySetups[ BodyIndex ];
+
+	if (!SharedData->bRunningSimulation && SharedData->GetSelectedConstraint())
+	{
+		UPhysicsConstraintTemplate* cs = SharedData->PhysicsAsset->ConstraintSetup[ SharedData->GetSelectedConstraint()->Index ];
+
+		if (cs->DefaultInstance.ConstraintBone1 == SharedBodySetup->BoneName)
+		{
+			return ConstraintBone1Color;
+		}
+		else if (cs->DefaultInstance.ConstraintBone2 == SharedBodySetup->BoneName)
+		{
+			return ConstraintBone2Color;
+		}
+	}
+
+	FPhysicsAssetEditorSharedData::FSelection Body(BodyIndex, PrimitiveType, PrimitiveIndex);
+
+	static FName SelectionColorName(TEXT("SelectionColor"));
+	const FSlateColor SelectionColor = FEditorStyle::GetSlateColor(SelectionColorName);
+	const FLinearColor SelectionColorLinear(SelectionColor.IsColorSpecified() ? SelectionColor.GetSpecifiedColor() : FLinearColor::White);
+	const FColor ElemSelectedColor = SelectionColorLinear.ToFColor(true);
+	const FColor ElemSelectedBodyColor = (SelectionColorLinear* 0.5f).ToFColor(true);
+
+	bool bInBody = false;
+	for(int32 i=0; i<SharedData->SelectedBodies.Num(); ++i)
+	{
+		if(BodyIndex == SharedData->SelectedBodies[i].Index)
+		{
+			bInBody = true;
+		}
+
+		if (Body == SharedData->SelectedBodies[i] && !SharedData->bRunningSimulation)
+		{
+			return ElemSelectedColor;
+		}
+	}
+
+	if(bInBody && !SharedData->bRunningSimulation)	//this primitive is in a body that's currently selected, but this primitive itself isn't selected
+	{
+		return ElemSelectedBodyColor;
+	}
+	
+	if (SharedData->bRunningSimulation)
+	{
+		const bool bIsSimulatedAtAll = SharedBodySetup->PhysicsType == PhysType_Simulated || (SharedBodySetup->PhysicsType == PhysType_Default && SharedData->EditorOptions->PhysicsBlend > 0.f);
+		if (!bIsSimulatedAtAll)
+		{
+			return FixedColor;
+		}
+	}
+	else
+	{
+		if (!SharedData->bRunningSimulation && SharedData->SelectedBodies.Num())
+		{
+			// If there is no collision with this body, use 'no collision material'.
+			if (SharedData->NoCollisionBodies.Find(BodyIndex) != INDEX_NONE)
+			{
+				return NoCollisionColor;
+			}
+		}
+	}
+
+	return BoneUnselectedColor;
+}
+
+UMaterialInterface* UPhysicsAssetEditorSkeletalMeshComponent::GetPrimitiveMaterial(int32 BodyIndex, EAggCollisionShape::Type PrimitiveType, int32 PrimitiveIndex)
+{
+	if (SharedData->bRunningSimulation)
+	{
+		return BoneUnselectedMaterial;
+	}
+
+	FPhysicsAssetEditorSharedData::FSelection Body(BodyIndex, PrimitiveType, PrimitiveIndex);
+
+	for(int32 i=0; i< SharedData->SelectedBodies.Num(); ++i)
+	{
+		if (Body == SharedData->SelectedBodies[i] && !SharedData->bRunningSimulation)
+		{
+			return ElemSelectedMaterial;
+		}
+	}
+
+	// If there is no collision with this body, use 'no collision material'.
+	if (SharedData->SelectedBodies.Num() && SharedData->NoCollisionBodies.Find(BodyIndex) != INDEX_NONE && !SharedData->bRunningSimulation)
+	{
+		return BoneNoCollisionMaterial;
+	}
+	else
+	{
+		return BoneUnselectedMaterial;
+	}
+
+}
+
+void UPhysicsAssetEditorSkeletalMeshComponent::RefreshBoneTransforms(FActorComponentTickFunction* TickFunction)
+{
+	Super::RefreshBoneTransforms(TickFunction);
+
+	// Horrible kludge, but we need to flip the buffer back here as we need to wait on the physics tick group.
+	// However UDebugSkelMeshComponent passes NULL to force non-threaded work, which assumes a flip is needed straight away
+	if(ShouldBlendPhysicsBones())
+	{
+		bNeedToFlipSpaceBaseBuffers = true;
+		FinalizeBoneTransform();
+		bNeedToFlipSpaceBaseBuffers = true;
+	}
+}
\ No newline at end of file
diff --git a/Engine/Source/Editor/PhysicsAssetEditor/Private/PhysicsAssetEditorSkeletalMeshComponent.h b/Engine/Source/Editor/PhysicsAssetEditor/Private/PhysicsAssetEditorSkeletalMeshComponent.h
new file mode 100644
--- /dev/null
+++ b/Engine/Source/Editor/PhysicsAssetEditor/Private/PhysicsAssetEditorSkeletalMeshComponent.h
@@ -0,0 +1,72 @@
+// Copyright 1998-2017 Epic Games, Inc. All Rights Reserved.
+
+#pragma once
+
+#include "CoreMinimal.h"
+#include "UObject/ObjectMacros.h"
+#include "PhysicsEngine/ShapeElem.h"
+#include "Animation/DebugSkelMeshComponent.h"
+#include "PhysicsAssetEditorSkeletalMeshComponent.generated.h"
+
+class FPrimitiveDrawInterface;
+class UMaterialInterface;
+class UMaterialInstanceDynamic;
+
+UCLASS()
+class UPhysicsAssetEditorSkeletalMeshComponent : public UDebugSkelMeshComponent
+{
+	GENERATED_UCLASS_BODY()
+
+	/** Data and methods shared across multiple classes */
+	class FPhysicsAssetEditorSharedData* SharedData;
+
+	// Draw colors
+
+	FColor BoneUnselectedColor;
+	FColor NoCollisionColor;
+	FColor FixedColor;
+	FColor ConstraintBone1Color;
+	FColor ConstraintBone2Color;
+	FColor HierarchyDrawColor;
+	FColor AnimSkelDrawColor;
+	float COMRenderSize;
+	float InfluenceLineLength;
+	FColor InfluenceLineColor;
+
+	// Materials
+
+	UPROPERTY(transient)
+	UMaterialInstanceDynamic* ElemSelectedMaterial;
+	UPROPERTY(transient)
+	UMaterialInstanceDynamic* BoneSelectedMaterial;
+	UPROPERTY(transient)
+	UMaterialInstanceDynamic* BoneUnselectedMaterial;
+	UPROPERTY(transient)
+	UMaterialInterface* BoneMaterialHit;
+	UPROPERTY(transient)
+	UMaterialInstanceDynamic* BoneNoCollisionMaterial;
+
+	/** Mesh-space matrices showing state of just animation (ie before physics) - useful for debugging! */
+	TArray<FTransform> AnimationSpaceBases;
+
+
+	/** UPrimitiveComponent interface */
+	virtual FPrimitiveSceneProxy* CreateSceneProxy() override;
+
+	/** USkinnedMeshComponent interface */
+	virtual void RefreshBoneTransforms(FActorComponentTickFunction* TickFunction = nullptr) override;
+	
+	/** Debug drawing */
+	void DebugDraw(const FSceneView* View, FPrimitiveDrawInterface* PDI);
+
+	/** Handles most of the rendering logic for this component */
+	void RenderAssetTools(const FSceneView* View, FPrimitiveDrawInterface* PDI);
+
+	/** Draws a constraint */
+	void DrawConstraint(int32 ConstraintIndex, const FSceneView* View, FPrimitiveDrawInterface* PDI, bool bDrawAsPoint);
+
+	/** Accessors/helper methods */
+	FTransform GetPrimitiveTransform(FTransform& BoneTM, int32 BodyIndex, EAggCollisionShape::Type PrimType, int32 PrimIndex, float Scale);
+	FColor GetPrimitiveColor(int32 BodyIndex, EAggCollisionShape::Type PrimitiveType, int32 PrimitiveIndex);
+	UMaterialInterface* GetPrimitiveMaterial(int32 BodyIndex, EAggCollisionShape::Type PrimitiveType, int32 PrimitiveIndex);
+};
diff --git a/Engine/Source/Editor/PhysicsAssetEditor/Private/PhysicsAssetGraph/PhysicsAssetConnectionDrawingPolicy.h b/Engine/Source/Editor/PhysicsAssetEditor/Private/PhysicsAssetGraph/PhysicsAssetConnectionDrawingPolicy.h
new file mode 100644
--- /dev/null
+++ b/Engine/Source/Editor/PhysicsAssetEditor/Private/PhysicsAssetGraph/PhysicsAssetConnectionDrawingPolicy.h
@@ -0,0 +1,23 @@
+// Copyright 1998-2017 Epic Games, Inc. All Rights Reserved.
+
+#pragma once
+
+#include "ConnectionDrawingPolicy.h"
+
+class FPhysicsAssetConnectionDrawingPolicy : public FConnectionDrawingPolicy
+{
+public:
+	FPhysicsAssetConnectionDrawingPolicy(int32 InBackLayerID, int32 InFrontLayerID, float InZoomFactor, const FSlateRect& InClippingRect, FSlateWindowElementList& InDrawElements)
+		: FConnectionDrawingPolicy(InBackLayerID, InFrontLayerID, InZoomFactor, InClippingRect, InDrawElements)
+	{
+		// We dont draw arrows here
+		ArrowImage = nullptr;
+		ArrowRadius = FVector2D(0.0f, 0.0f);
+	}
+
+	virtual FVector2D ComputeSplineTangent(const FVector2D& Start, const FVector2D& End) const override
+	{
+		const int32 Tension = FMath::Abs<int32>(Start.X - End.X);
+		return Tension * FVector2D(1.0f, 0);
+	}
+};
diff --git a/Engine/Source/Editor/PhysicsAssetEditor/Private/PhysicsAssetGraph/PhysicsAssetGraphSchema.h b/Engine/Source/Editor/PhysicsAssetEditor/Private/PhysicsAssetGraph/PhysicsAssetGraphSchema.h
new file mode 100644
--- /dev/null
+++ b/Engine/Source/Editor/PhysicsAssetEditor/Private/PhysicsAssetGraph/PhysicsAssetGraphSchema.h
@@ -0,0 +1,39 @@
+// Copyright 1998-2016 Epic Games, Inc. All Rights Reserved.
+
+#pragma once
+
+#include "EdGraph/EdGraphSchema.h"
+#include "PhysicsEngine/ShapeElem.h"
+#include "PhysicsAssetGraphSchema.generated.h"
+
+class UPhysicsAsset;
+class UPhysicsAssetGraph;
+class UPhysicsAssetGraphNode;
+class UPhysicsAssetGraphNode_Bone;
+class UPhysicsAssetGraphNode_Constraint;
+class USkeletalBodySetup;
+class UPhysicsConstraintTemplate;
+
+UCLASS()
+class UPhysicsAssetGraphSchema : public UEdGraphSchema
+{
+	GENERATED_UCLASS_BODY()
+
+public:
+	// UEdGraphSchema interface
+	virtual void GetContextMenuActions(const UEdGraph* CurrentGraph, const UEdGraphNode* InGraphNode, const UEdGraphPin* InGraphPin, class FMenuBuilder* MenuBuilder, bool bIsDebugging) const override;
+	virtual FLinearColor GetPinTypeColor(const FEdGraphPinType& PinType) const override;
+	virtual void BreakPinLinks(UEdGraphPin& TargetPin, bool bSendsNodeNotifcation) const override;
+	virtual void BreakSinglePinLink(UEdGraphPin* SourcePin, UEdGraphPin* TargetPin) override;
+	virtual FPinConnectionResponse MovePinLinks(UEdGraphPin& MoveFromPin, UEdGraphPin& MoveToPin, bool bIsItermeadiateMove = false) const override;
+	virtual FPinConnectionResponse CopyPinLinks(UEdGraphPin& CopyFromPin, UEdGraphPin& CopyToPin, bool bIsItermeadiateCopy = false) const override;
+	virtual class FConnectionDrawingPolicy* CreateConnectionDrawingPolicy(int32 InBackLayerID, int32 InFrontLayerID, float InZoomFactor, const FSlateRect& InClippingRect, class FSlateWindowElementList& InDrawElements, class UEdGraph* InGraphObj) const override;
+	virtual bool ShouldAlwaysPurgeOnModification() const override { return false; }
+	// End of UEdGraphSchema interface
+
+	virtual UPhysicsAssetGraphNode_Bone* CreateGraphNodesForBone(UPhysicsAssetGraph* InGraph, USkeletalBodySetup* InBodySetup, int32 InBodyIndex, UPhysicsAsset* InPhysicsAsset) const;
+	virtual UPhysicsAssetGraphNode_Constraint* CreateGraphNodesForConstraint(UPhysicsAssetGraph* InGraph, UPhysicsConstraintTemplate* InConstraint, int32 InConstraintIndex, UPhysicsAsset* InPhysicsAsset) const;
+
+	virtual void LayoutNodes(UPhysicsAssetGraph* InGraph, UPhysicsAsset* InPhysicsAsset) const;
+};
+
diff --git a/Engine/Source/Editor/PhysicsAssetEditor/Private/PhysicsAssetGraph/SBoneNode.cpp b/Engine/Source/Editor/PhysicsAssetEditor/Private/PhysicsAssetGraph/SBoneNode.cpp
new file mode 100644
--- /dev/null
+++ b/Engine/Source/Editor/PhysicsAssetEditor/Private/PhysicsAssetGraph/SBoneNode.cpp
@@ -0,0 +1,21 @@
+// Copyright 1998-2016 Epic Games, Inc. All Rights Reserved.
+
+#include "SBoneNode.h"
+#include "PhysicsAssetGraphNode_Bone.h"
+#include "SGraphPin.h"
+#include "SInlineEditableTextBlock.h"
+#include "PhysicsEngine/ShapeElem.h"
+#include "SImage.h"
+#include "SButton.h"
+#include "PhysicsEngine/AggregateGeom.h"
+#include "PhysicsEngine/PhysicsAsset.h"
+
+#define LOCTEXT_NAMESPACE "SBoneNode"
+
+void SBoneNode::Construct(const FArguments& InArgs, UPhysicsAssetGraphNode_Bone* InNode)
+{
+	SPhysicsAssetGraphNode::Construct(SPhysicsAssetGraphNode::FArguments(), InNode);
+	UpdateGraphNode();
+}
+
+#undef LOCTEXT_NAMESPACE
\ No newline at end of file
diff --git a/Engine/Source/Editor/PhysicsAssetEditor/Private/PhysicsAssetGraph/SConstraintNode.cpp b/Engine/Source/Editor/PhysicsAssetEditor/Private/PhysicsAssetGraph/SConstraintNode.cpp
new file mode 100644
--- /dev/null
+++ b/Engine/Source/Editor/PhysicsAssetEditor/Private/PhysicsAssetGraph/SConstraintNode.cpp
@@ -0,0 +1,12 @@
+// Copyright 1998-2016 Epic Games, Inc. All Rights Reserved.
+
+#include "SConstraintNode.h"
+#include "PhysicsAssetGraphNode_Constraint.h"
+#include "SGraphPin.h"
+
+void SConstraintNode::Construct(const FArguments& InArgs, UPhysicsAssetGraphNode_Constraint* InNode)
+{
+	SPhysicsAssetGraphNode::Construct(SPhysicsAssetGraphNode::FArguments(), InNode);
+
+	UpdateGraphNode();
+}
diff --git a/Engine/Source/Editor/PhysicsAssetEditor/Private/SkeletonTreePhysicsBodyItem.cpp b/Engine/Source/Editor/PhysicsAssetEditor/Private/SkeletonTreePhysicsBodyItem.cpp
new file mode 100644
--- /dev/null
+++ b/Engine/Source/Editor/PhysicsAssetEditor/Private/SkeletonTreePhysicsBodyItem.cpp
@@ -0,0 +1,65 @@
+// Copyright 1998-2017 Epic Games, Inc. All Rights Reserved.
+
+#include "SkeletonTreePhysicsBodyItem.h"
+#include "Widgets/Text/STextBlock.h"
+#include "Widgets/Images/SImage.h"
+#include "EditorStyleSet.h"
+
+#define LOCTEXT_NAMESPACE "FSkeletonTreePhysicsBodyItem"
+
+void FSkeletonTreePhysicsBodyItem::GenerateWidgetForNameColumn( TSharedPtr< SHorizontalBox > Box, const TAttribute<FText>& FilterText, FIsSelected InIsSelected )
+{
+	Box->AddSlot()
+	.AutoWidth()
+	.Padding(FMargin(0.0f, 1.0f))
+	[
+		SNew( SImage )
+		.ColorAndOpacity(FSlateColor::UseForeground())
+		.Image(this, &FSkeletonTreePhysicsBodyItem::GetBrush)
+	];
+
+	Box->AddSlot()
+	.AutoWidth()
+	.Padding(2, 0, 0, 0)
+	[
+		SNew(STextBlock)
+		.ColorAndOpacity(this, &FSkeletonTreePhysicsBodyItem::GetBodyTextColor)
+		.Text(FText::FromName(BoneName))
+		.HighlightText(FilterText)
+		.Font(FEditorStyle::GetFontStyle("PhysicsAssetEditor.Tree.Font"))
+		.ToolTipText(FText::Format(LOCTEXT("BodyTooltip", "Aggregate physics body for bone '{0}'. Bodies can consist of multiple shapes."), FText::FromName(BoneName)))
+	];
+}
+
+TSharedRef< SWidget > FSkeletonTreePhysicsBodyItem::GenerateWidgetForDataColumn(const FName& DataColumnName)
+{
+	return SNullWidget::NullWidget;
+}
+
+const FSlateBrush* FSkeletonTreePhysicsBodyItem::GetBrush() const
+{
+	return BodySetup->PhysicsType == EPhysicsType::PhysType_Kinematic ? FEditorStyle::GetBrush("PhysicsAssetEditor.Tree.KinematicBody") : FEditorStyle::GetBrush("PhysicsAssetEditor.Tree.Body");
+}
+
+FSlateColor FSkeletonTreePhysicsBodyItem::GetBodyTextColor() const
+{
+	FLinearColor Color(1.0f, 1.0f, 1.0f);
+
+	const bool bInCurrentProfile = BodySetup->GetCurrentPhysicalAnimationProfileName() == NAME_None || BodySetup->FindPhysicalAnimationProfile(BodySetup->GetCurrentPhysicalAnimationProfileName()) != nullptr;
+
+	if (FilterResult == ESkeletonTreeFilterResult::ShownDescendant)
+	{
+		Color = FLinearColor::Gray * 0.5f;
+	}
+
+	if(bInCurrentProfile)
+	{
+		return FSlateColor(Color);
+	}
+	else
+	{
+		return FSlateColor(Color.Desaturate(0.5f));
+	}
+}
+
+#undef LOCTEXT_NAMESPACE
diff --git a/Engine/Source/Editor/PhysicsAssetEditor/Private/SkeletonTreePhysicsBodyItem.h b/Engine/Source/Editor/PhysicsAssetEditor/Private/SkeletonTreePhysicsBodyItem.h
new file mode 100644
--- /dev/null
+++ b/Engine/Source/Editor/PhysicsAssetEditor/Private/SkeletonTreePhysicsBodyItem.h
@@ -0,0 +1,59 @@
+// Copyright 1998-2017 Epic Games, Inc. All Rights Reserved.
+
+#pragma once
+
+#include "CoreMinimal.h"
+#include "Styling/SlateColor.h"
+#include "Input/Reply.h"
+#include "Widgets/SWidget.h"
+#include "Widgets/Views/STableViewBase.h"
+#include "Widgets/Views/STableRow.h"
+#include "SkeletonTreeItem.h"
+#include "PhysicsEngine/PhysicsAsset.h"
+
+class FSkeletonTreePhysicsBodyItem : public FSkeletonTreeItem
+{
+public:
+	SKELETON_TREE_ITEM_TYPE(FSkeletonTreePhysicsBodyItem, FSkeletonTreeItem)
+
+	FSkeletonTreePhysicsBodyItem(USkeletalBodySetup* InBodySetup, int32 InBodySetupIndex, const FName& InBoneName, bool bInHasBodySetup, bool bInHasShapes, const TSharedRef<class ISkeletonTree>& InSkeletonTree)
+		: FSkeletonTreeItem(InSkeletonTree)
+		, BodySetup(InBodySetup)
+		, BodySetupIndex(InBodySetupIndex)
+		, BoneName(InBoneName)
+		, bHasBodySetup(bInHasBodySetup)
+		, bHasShapes(bInHasShapes)
+	{}
+
+	/** ISkeletonTreeItem interface */
+	virtual void GenerateWidgetForNameColumn(TSharedPtr< SHorizontalBox > Box, const TAttribute<FText>& FilterText, FIsSelected InIsSelected) override;
+	virtual TSharedRef< SWidget > GenerateWidgetForDataColumn(const FName& DataColumnName) override;	
+	virtual FName GetRowItemName() const override { return BoneName; }
+	virtual UObject* GetObject() const override { return BodySetup; }
+
+	/** Get the index of the body setup in the physics asset */
+	int32 GetBodySetupIndex() const { return BodySetupIndex; }
+
+private:
+	/** Gets the icon to display for this body */
+	const FSlateBrush* GetBrush() const;
+
+	/** Gets the color to display the item's text */
+	FSlateColor GetBodyTextColor() const;
+
+private:
+	/** The body setup we are representing */
+	USkeletalBodySetup* BodySetup;
+
+	/** The index of the body setup in the physics asset */
+	int32 BodySetupIndex;
+
+	/** The name of the bone that this body is bound to */
+	FName BoneName;
+
+	/** Whether there is a body set up for this bone */
+	bool bHasBodySetup;
+
+	/** Whether this body has shapes */
+	bool bHasShapes;
+};
diff --git a/Engine/Source/Editor/PhysicsAssetEditor/Private/SkeletonTreePhysicsConstraintItem.cpp b/Engine/Source/Editor/PhysicsAssetEditor/Private/SkeletonTreePhysicsConstraintItem.cpp
new file mode 100644
--- /dev/null
+++ b/Engine/Source/Editor/PhysicsAssetEditor/Private/SkeletonTreePhysicsConstraintItem.cpp
@@ -0,0 +1,55 @@
+// Copyright 1998-2017 Epic Games, Inc. All Rights Reserved.
+
+#include "SkeletonTreePhysicsConstraintItem.h"
+#include "Widgets/Text/STextBlock.h"
+#include "Widgets/Images/SImage.h"
+#include "EditorStyleSet.h"
+
+#define LOCTEXT_NAMESPACE "FSkeletonTreePhysicsConstraintItem"
+
+void FSkeletonTreePhysicsConstraintItem::GenerateWidgetForNameColumn( TSharedPtr< SHorizontalBox > Box, const TAttribute<FText>& FilterText, FIsSelected InIsSelected )
+{
+	Box->AddSlot()
+	.AutoWidth()
+	.Padding(FMargin(0.0f, 1.0f))
+	[
+		SNew( SImage )
+		.ColorAndOpacity(FSlateColor::UseForeground())
+		.Image(FEditorStyle::GetBrush("PhysicsAssetEditor.Tree.Constraint"))
+	];
+
+	const FConstraintInstance& ConstraintInstance = Constraint->DefaultInstance;
+
+	Box->AddSlot()
+	.AutoWidth()
+	.Padding(2, 0, 0, 0)
+	[
+		SNew(STextBlock)
+		.ColorAndOpacity(this, &FSkeletonTreePhysicsConstraintItem::GetConstraintTextColor)
+		.Text(FText::Format(LOCTEXT("ConstraintNameFormat", "{0} : {1} Constraint"), FText::FromName(ConstraintInstance.ConstraintBone1), FText::FromName(ConstraintInstance.ConstraintBone2)))
+		.HighlightText(FilterText)
+		.Font(FEditorStyle::GetFontStyle("PhysicsAssetEditor.Tree.Font"))
+		.ToolTipText(FText::Format(LOCTEXT("ConstraintTooltip", "Constraint linking '{0}' and '{1}'"), FText::FromName(ConstraintInstance.ConstraintBone1), FText::FromName(ConstraintInstance.ConstraintBone2)))
+	];
+}
+
+TSharedRef< SWidget > FSkeletonTreePhysicsConstraintItem::GenerateWidgetForDataColumn(const FName& DataColumnName)
+{
+	return SNullWidget::NullWidget;
+}
+
+FSlateColor FSkeletonTreePhysicsConstraintItem::GetConstraintTextColor() const
+{
+	const FLinearColor Color(1.0f, 1.0f, 1.0f);
+	const bool bInCurrentProfile = Constraint->GetCurrentConstraintProfileName() == NAME_None || Constraint->ContainsConstraintProfile(Constraint->GetCurrentConstraintProfileName());
+	if(bInCurrentProfile)
+	{
+		return FSlateColor(Color);
+	}
+	else
+	{
+		return FSlateColor(Color.Desaturate(0.5f));
+	}
+}
+
+#undef LOCTEXT_NAMESPACE
diff --git a/Engine/Source/Editor/PhysicsAssetEditor/Private/SkeletonTreePhysicsConstraintItem.h b/Engine/Source/Editor/PhysicsAssetEditor/Private/SkeletonTreePhysicsConstraintItem.h
new file mode 100644
--- /dev/null
+++ b/Engine/Source/Editor/PhysicsAssetEditor/Private/SkeletonTreePhysicsConstraintItem.h
@@ -0,0 +1,54 @@
+// Copyright 1998-2017 Epic Games, Inc. All Rights Reserved.
+
+#pragma once
+
+#include "CoreMinimal.h"
+#include "Styling/SlateColor.h"
+#include "Input/Reply.h"
+#include "Widgets/SWidget.h"
+#include "Widgets/Views/STableViewBase.h"
+#include "Widgets/Views/STableRow.h"
+#include "SkeletonTreeItem.h"
+#include "PhysicsEngine/PhysicsAsset.h"
+#include "PhysicsEngine/PhysicsConstraintTemplate.h"
+
+class FSkeletonTreePhysicsConstraintItem : public FSkeletonTreeItem
+{
+public:
+	SKELETON_TREE_ITEM_TYPE(FSkeletonTreePhysicsConstraintItem, FSkeletonTreeItem)
+
+	FSkeletonTreePhysicsConstraintItem(UPhysicsConstraintTemplate* InConstraint, int32 InConstraintIndex, const FName& InBoneName, const TSharedRef<class ISkeletonTree>& InSkeletonTree)
+		: FSkeletonTreeItem(InSkeletonTree)
+		, Constraint(InConstraint)
+		, ConstraintIndex(InConstraintIndex)
+		, BoneName(InBoneName)
+	{
+		const FConstraintInstance& ConstraintInstance = Constraint->DefaultInstance;
+		OtherBoneName = ConstraintInstance.ConstraintBone1 == BoneName ? ConstraintInstance.ConstraintBone2 : ConstraintInstance.ConstraintBone1;
+	}
+
+	/** ISkeletonTreeItem interface */
+	virtual void GenerateWidgetForNameColumn(TSharedPtr< SHorizontalBox > Box, const TAttribute<FText>& FilterText, FIsSelected InIsSelected) override;
+	virtual TSharedRef< SWidget > GenerateWidgetForDataColumn(const FName& DataColumnName) override;	
+	virtual FName GetRowItemName() const override { return OtherBoneName; }
+	virtual UObject* GetObject() const override { return Constraint; }
+
+	/** Get the index of the constraint in the physics asset */
+	int32 GetConstraintIndex() const { return ConstraintIndex; }
+
+private:
+	FSlateColor GetConstraintTextColor() const;
+
+private:
+	/** The constraint we are representing */
+	UPhysicsConstraintTemplate* Constraint;
+
+	/** The index of the body setup in the physics asset */
+	int32 ConstraintIndex;
+
+	/** The constrained bone we are parented to in the tree */
+	FName BoneName;
+
+	/** The constrained bone we are not parented to in the tree */
+	FName OtherBoneName;
+};
diff --git a/Engine/Source/Editor/PhysicsAssetEditor/Private/SkeletonTreePhysicsShapeItem.cpp b/Engine/Source/Editor/PhysicsAssetEditor/Private/SkeletonTreePhysicsShapeItem.cpp
new file mode 100644
--- /dev/null
+++ b/Engine/Source/Editor/PhysicsAssetEditor/Private/SkeletonTreePhysicsShapeItem.cpp
@@ -0,0 +1,72 @@
+// Copyright 1998-2017 Epic Games, Inc. All Rights Reserved.
+
+#include "SkeletonTreePhysicsShapeItem.h"
+#include "Widgets/Text/STextBlock.h"
+#include "Widgets/Images/SImage.h"
+#include "EditorStyleSet.h"
+
+#define LOCTEXT_NAMESPACE "FSkeletonTreePhysicsShapeItem"
+
+FSkeletonTreePhysicsShapeItem::FSkeletonTreePhysicsShapeItem(USkeletalBodySetup* InBodySetup, const FName& InBoneName, int32 InBodySetupIndex, EAggCollisionShape::Type InShapeType, int32 InShapeIndex, const TSharedRef<class ISkeletonTree>& InSkeletonTree)
+	: FSkeletonTreeItem(InSkeletonTree)
+	, BodySetup(InBodySetup)
+	, BodySetupIndex(InBodySetupIndex)
+	, ShapeType(InShapeType)
+	, ShapeIndex(InShapeIndex)
+{
+	switch (ShapeType)
+	{
+	case EAggCollisionShape::Sphere:
+		ShapeBrush = FEditorStyle::GetBrush("PhysicsAssetEditor.Tree.Sphere");
+		Label = *FText::Format(LOCTEXT("SphereLabel", "{0} Sphere {1}"), FText::FromName(InBoneName), FText::AsNumber(ShapeIndex)).ToString();
+		break;
+	case EAggCollisionShape::Box:
+		ShapeBrush = FEditorStyle::GetBrush("PhysicsAssetEditor.Tree.Box");
+		Label = *FText::Format(LOCTEXT("BoxLabel", "{0} Box {1}"), FText::FromName(InBoneName), FText::AsNumber(ShapeIndex)).ToString();
+		break;
+	case EAggCollisionShape::Sphyl:
+		ShapeBrush = FEditorStyle::GetBrush("PhysicsAssetEditor.Tree.Sphyl");
+		Label = *FText::Format(LOCTEXT("CapsuleLabel", "{0} Capsule {1}"), FText::FromName(InBoneName), FText::AsNumber(ShapeIndex)).ToString();
+		break;
+	case EAggCollisionShape::Convex:
+		ShapeBrush = FEditorStyle::GetBrush("PhysicsAssetEditor.Tree.Convex");
+		Label = *FText::Format(LOCTEXT("ConvexLabel", "{0} Convex {1}"), FText::FromName(InBoneName), FText::AsNumber(ShapeIndex)).ToString();
+		break;
+	default:
+		check(false);
+		break;
+	}
+
+	RowItemName = *Label.ToString();
+}
+
+void FSkeletonTreePhysicsShapeItem::GenerateWidgetForNameColumn( TSharedPtr< SHorizontalBox > Box, const TAttribute<FText>& FilterText, FIsSelected InIsSelected )
+{
+	Box->AddSlot()
+	.AutoWidth()
+	.Padding(FMargin(0.0f, 1.0f))
+	[
+		SNew( SImage )
+		.ColorAndOpacity(FSlateColor::UseForeground())
+		.Image(ShapeBrush)
+	];
+
+	Box->AddSlot()
+	.AutoWidth()
+	.Padding(2, 0, 0, 0)
+	[
+		SNew(STextBlock)
+		.ColorAndOpacity(FSlateColor::UseForeground())
+		.Text(FText::FromName(Label))
+		.HighlightText(FilterText)
+		.Font(FEditorStyle::GetFontStyle("PhysicsAssetEditor.Tree.Font"))
+		.ToolTipText(FText::FromName(Label))
+	];
+}
+
+TSharedRef< SWidget > FSkeletonTreePhysicsShapeItem::GenerateWidgetForDataColumn(const FName& DataColumnName)
+{
+	return SNullWidget::NullWidget;
+}
+
+#undef LOCTEXT_NAMESPACE
diff --git a/Engine/Source/Editor/PhysicsAssetEditor/Public/IPhysicsAssetEditor.h b/Engine/Source/Editor/PhysicsAssetEditor/Public/IPhysicsAssetEditor.h
new file mode 100644
--- /dev/null
+++ b/Engine/Source/Editor/PhysicsAssetEditor/Public/IPhysicsAssetEditor.h
@@ -0,0 +1,16 @@
+// Copyright 1998-2017 Epic Games, Inc. All Rights Reserved.
+
+#pragma once
+
+#include "CoreMinimal.h"
+#include "WorkflowCentricApplication.h"
+
+
+/*-----------------------------------------------------------------------------
+   IPhysicsAssetEditor
+-----------------------------------------------------------------------------*/
+
+class IPhysicsAssetEditor : public FWorkflowCentricApplication
+{
+
+};
diff --git a/Engine/Source/Editor/PropertyEditor/Private/SDetailTableRowBase.h b/Engine/Source/Editor/PropertyEditor/Private/SDetailTableRowBase.h
--- a/Engine/Source/Editor/PropertyEditor/Private/SDetailTableRowBase.h
+++ b/Engine/Source/Editor/PropertyEditor/Private/SDetailTableRowBase.h
@@ -24,7 +24,7 @@ public:
 	{
 		if( OwnerTreeNode.IsValid() && MouseEvent.GetEffectingButton() == EKeys::RightMouseButton && !StaticCastSharedRef<STableViewBase>( OwnerTablePtr.Pin()->AsWidget() )->IsRightClickScrolling() )
 		{
-			FMenuBuilder MenuBuilder( true, NULL );
+			FMenuBuilder MenuBuilder( true, nullptr, nullptr, true );
 
 			FDetailNodeList VisibleChildren;
 			OwnerTreeNode.Pin()->GetChildren( VisibleChildren );
