Fix shadow variable warning on Linux

diff --git a/Engine/Source/Editor/PhysicsAssetEditor/Private/PhysicsAssetEditor.cpp b/Engine/Source/Editor/PhysicsAssetEditor/Private/PhysicsAssetEditor.cpp
--- a/Engine/Source/Editor/PhysicsAssetEditor/Private/PhysicsAssetEditor.cpp
+++ b/Engine/Source/Editor/PhysicsAssetEditor/Private/PhysicsAssetEditor.cpp
@@ -1012,9 +1012,9 @@ void FPhysicsAssetEditor::BuildMenuWidgetBody(FMenuBuilder& InMenuBuilder)
 		InMenuBuilder.AddSubMenu(
 			Commands.ApplyPhysicalMaterial->GetLabel(), 
 			Commands.ApplyPhysicalMaterial->GetDescription(), 
-			FNewMenuDelegate::CreateLambda([this](FMenuBuilder& InMenuBuilder)
+			FNewMenuDelegate::CreateLambda([this](FMenuBuilder& InSubMenuBuilder)
 			{
-				InMenuBuilder.AddWidget(BuildPhysicalMaterialAssetPicker(), FText(), true);
+				InSubMenuBuilder.AddWidget(BuildPhysicalMaterialAssetPicker(), FText(), true);
 			}),
 			FUIAction(FExecuteAction(), FCanExecuteAction::CreateSP(this, &FPhysicsAssetEditor::IsNotSimulation)),
 			NAME_None,
