Prevent undo/redo breaking when moving both a constraint and a body at the same time

diff --git a/Engine/Source/Editor/PhysicsAssetEditor/Private/PhysicsAssetEditorEditMode.cpp b/Engine/Source/Editor/PhysicsAssetEditor/Private/PhysicsAssetEditorEditMode.cpp
--- a/Engine/Source/Editor/PhysicsAssetEditor/Private/PhysicsAssetEditorEditMode.cpp
+++ b/Engine/Source/Editor/PhysicsAssetEditor/Private/PhysicsAssetEditorEditMode.cpp
@@ -108,11 +108,22 @@ bool FPhysicsAssetEditorEditMode::StartTracking(FEditorViewportClient* InViewpor
 	const EAxisList::Type CurrentAxis = InViewportClient->GetCurrentWidgetAxis();
 	if(CurrentAxis != EAxisList::None)
 	{
+		if(SharedData->GetSelectedBody() || SharedData->GetSelectedConstraint())
+		{
+			if(SharedData->GetSelectedBody())
+			{
+				GEditor->BeginTransaction(NSLOCTEXT("UnrealEd", "MoveElement", "Move Element"));
+			}
+			else
+			{
+				GEditor->BeginTransaction(NSLOCTEXT("UnrealEd", "MoveConstraint", "Move Constraint"));
+			}
+		}
+
 		// If releasing the mouse button, check we are done manipulating
 		check(!SharedData->bManipulating);
 		if (SharedData->GetSelectedBody())
 		{
-			GEditor->BeginTransaction(NSLOCTEXT("UnrealEd", "MoveElement", "Move Element"));
 			for (int32 i = 0; i<SharedData->SelectedBodies.Num(); ++i)
 			{
 				SharedData->PhysicsAsset->SkeletalBodySetups[SharedData->SelectedBodies[i].Index]->Modify();
@@ -124,8 +135,6 @@ bool FPhysicsAssetEditorEditMode::StartTracking(FEditorViewportClient* InViewpor
 
 		if (SharedData->GetSelectedConstraint())
 		{
-			GEditor->BeginTransaction(NSLOCTEXT("UnrealEd", "MoveConstraint", "Move Constraint"));
-
 			const int32 Count = SharedData->SelectedConstraints.Num();
 			StartManRelConTM.SetNumUninitialized(Count);
 			StartManParentConTM.SetNumUninitialized(Count);
