Fix crash undo-ing primitive regeneration while simulation is in progress & stopping simulation

diff --git a/Engine/Source/Editor/PhysicsAssetEditor/Private/PhysicsAssetEditorSharedData.cpp b/Engine/Source/Editor/PhysicsAssetEditor/Private/PhysicsAssetEditorSharedData.cpp
--- a/Engine/Source/Editor/PhysicsAssetEditor/Private/PhysicsAssetEditorSharedData.cpp
+++ b/Engine/Source/Editor/PhysicsAssetEditor/Private/PhysicsAssetEditorSharedData.cpp
@@ -1638,11 +1638,6 @@ TSharedRef<SWidget> FPhysicsAssetEditorSharedData::CreateGenerateBodiesWidget(co
 
 void FPhysicsAssetEditorSharedData::PostUndo()
 {
-	if (bRunningSimulation)
-	{
-		return;
-	}
-
 	bool bInvalidSelection = false;
 
 	for (int32 BodyIndex = 0; BodyIndex < SelectedBodies.Num() && bInvalidSelection == false; ++BodyIndex)
