Fix out of bounds access when meshes have bones removed under physics assets

diff --git a/Engine/Source/Editor/PhysicsAssetEditor/Private/PhysicsAssetEditorEditMode.cpp b/Engine/Source/Editor/PhysicsAssetEditor/Private/PhysicsAssetEditorEditMode.cpp
--- a/Engine/Source/Editor/PhysicsAssetEditor/Private/PhysicsAssetEditorEditMode.cpp
+++ b/Engine/Source/Editor/PhysicsAssetEditor/Private/PhysicsAssetEditorEditMode.cpp
@@ -604,12 +604,14 @@ bool FPhysicsAssetEditorEditMode::GetCustomDrawingCoordinateSystem(FMatrix& InMa
 	if (SharedData->GetSelectedBody())
 	{
 		int32 BoneIndex = SharedData->EditorSkelComp->GetBoneIndex(SharedData->PhysicsAsset->SkeletalBodySetups[SharedData->GetSelectedBody()->Index]->BoneName);
+		if(BoneIndex != INDEX_NONE)
+		{
+			FTransform BoneTM = SharedData->EditorSkelComp->GetBoneTransform(BoneIndex);
+			BoneTM.RemoveScaling();
 
-		FTransform BoneTM = SharedData->EditorSkelComp->GetBoneTransform(BoneIndex);
-		BoneTM.RemoveScaling();
-
-		InMatrix = SharedData->EditorSkelComp->GetPrimitiveTransform(BoneTM, SharedData->GetSelectedBody()->Index, SharedData->GetSelectedBody()->PrimitiveType, SharedData->GetSelectedBody()->PrimitiveIndex, 1.f).ToMatrixNoScale().RemoveTranslation();
-		return true;
+			InMatrix = SharedData->EditorSkelComp->GetPrimitiveTransform(BoneTM, SharedData->GetSelectedBody()->Index, SharedData->GetSelectedBody()->PrimitiveType, SharedData->GetSelectedBody()->PrimitiveIndex, 1.f).ToMatrixNoScale().RemoveTranslation();
+			return true;
+		}
 	}
 	else if (SharedData->GetSelectedConstraint())
 	{
@@ -631,12 +633,14 @@ FVector FPhysicsAssetEditorEditMode::GetWidgetLocation() const
 	if (SharedData->GetSelectedBody())
 	{
 		int32 BoneIndex = SharedData->EditorSkelComp->GetBoneIndex(SharedData->PhysicsAsset->SkeletalBodySetups[SharedData->GetSelectedBody()->Index]->BoneName);
+		if(BoneIndex != INDEX_NONE)
+		{
+			FTransform BoneTM = SharedData->EditorSkelComp->GetBoneTransform(BoneIndex);
+			const float Scale = BoneTM.GetScale3D().GetAbsMax();
+			BoneTM.RemoveScaling();
 
-		FTransform BoneTM = SharedData->EditorSkelComp->GetBoneTransform(BoneIndex);
-		const float Scale = BoneTM.GetScale3D().GetAbsMax();
-		BoneTM.RemoveScaling();
-
-		return SharedData->EditorSkelComp->GetPrimitiveTransform(BoneTM, SharedData->GetSelectedBody()->Index, SharedData->GetSelectedBody()->PrimitiveType, SharedData->GetSelectedBody()->PrimitiveIndex, Scale).GetTranslation();
+			return SharedData->EditorSkelComp->GetPrimitiveTransform(BoneTM, SharedData->GetSelectedBody()->Index, SharedData->GetSelectedBody()->PrimitiveType, SharedData->GetSelectedBody()->PrimitiveIndex, Scale).GetTranslation();
+		}
 	}
 	else if (SharedData->GetSelectedConstraint())
 	{
diff --git a/Engine/Source/Runtime/Engine/Private/Components/SkinnedMeshComponent.cpp b/Engine/Source/Runtime/Engine/Private/Components/SkinnedMeshComponent.cpp
--- a/Engine/Source/Runtime/Engine/Private/Components/SkinnedMeshComponent.cpp
+++ b/Engine/Source/Runtime/Engine/Private/Components/SkinnedMeshComponent.cpp
@@ -1047,7 +1047,8 @@ FMatrix USkinnedMeshComponent::GetBoneMatrix(int32 BoneIdx) const
 	}
 	else
 	{
-		if(GetNumComponentSpaceTransforms() && BoneIdx < GetNumComponentSpaceTransforms() )
+		const int32 NumTransforms = GetNumComponentSpaceTransforms();
+		if(BoneIdx >= 0 && BoneIdx < NumTransforms)
 		{
 			return GetComponentSpaceTransforms()[BoneIdx].ToMatrixWithScale() * GetComponentTransform().ToMatrixWithScale();
 		}
@@ -1108,7 +1109,7 @@ FTransform USkinnedMeshComponent::GetBoneTransform(int32 BoneIdx, const FTransfo
 	else
 	{
 		const int32 NumTransforms = GetNumComponentSpaceTransforms();
-		if(NumTransforms > 0 && BoneIdx < NumTransforms)
+		if(BoneIdx >= 0 && BoneIdx < NumTransforms)
 		{
 			return GetComponentSpaceTransforms()[BoneIdx] * LocalToWorld;
 		}
