Disallow move elements and constraints in phat when simulating.

diff --git a/Engine/Source/Editor/PhysicsAssetEditor/Private/PhysicsAssetEditorEditMode.cpp b/Engine/Source/Editor/PhysicsAssetEditor/Private/PhysicsAssetEditorEditMode.cpp
--- a/Engine/Source/Editor/PhysicsAssetEditor/Private/PhysicsAssetEditorEditMode.cpp
+++ b/Engine/Source/Editor/PhysicsAssetEditor/Private/PhysicsAssetEditorEditMode.cpp
@@ -110,7 +110,7 @@ void FPhysicsAssetEditorEditMode::GetOnScreenDebugInfo(TArray<FText>& OutDebugIn
 bool FPhysicsAssetEditorEditMode::StartTracking(FEditorViewportClient* InViewportClient, FViewport* InViewport)
 {
 	const EAxisList::Type CurrentAxis = InViewportClient->GetCurrentWidgetAxis();
-	if(!SharedData->bManipulating && CurrentAxis != EAxisList::None)
+	if(!SharedData->bRunningSimulation && !SharedData->bManipulating && CurrentAxis != EAxisList::None)
 	{
 		if(SharedData->GetSelectedBody() || SharedData->GetSelectedConstraint())
 		{
