Fix Welding Mannequin's ball_l and foot_l collider crashing in physics asset editor. #rb none #jira UE-69323

diff --git a/Engine/Source/Editor/PhysicsAssetEditor/Private/PhysicsAssetEditorSharedData.cpp b/Engine/Source/Editor/PhysicsAssetEditor/Private/PhysicsAssetEditorSharedData.cpp
--- a/Engine/Source/Editor/PhysicsAssetEditor/Private/PhysicsAssetEditorSharedData.cpp
+++ b/Engine/Source/Editor/PhysicsAssetEditor/Private/PhysicsAssetEditorSharedData.cpp
@@ -944,7 +944,10 @@ bool FPhysicsAssetEditorSharedData::WeldSelectedBodies(bool bWeld /* = true */)
 	// update the tree
 	HierarchyChangedEvent.Broadcast();
 
-	// Body index may have changed, so we re-find it.
+	// Just to be safe - deselect any selected constraints
+	ClearSelectedConstraints();
+	ClearSelectedBody(); // Previous selection is invalid because child no longer has same index.
+
 	int32 BodyIndex = PhysicsAsset->FindBodyIndex(ParentBoneName);
 	FSelection SelectionParent(BodyIndex, ParentPrimitiveType, ParentPrimitiveIndex);
 	SetSelectedBody(SelectionParent, true); // This redraws the viewport as well...
@@ -952,8 +955,6 @@ bool FPhysicsAssetEditorSharedData::WeldSelectedBodies(bool bWeld /* = true */)
 	FSelection SelectionChild(BodyIndex, ChildPrimitiveType, ChildPrimitiveIndex);
 	SetSelectedBody(SelectionChild, true); // This redraws the viewport as well...
 
-	// Just to be safe - deselect any selected constraints
-	ClearSelectedConstraints();
 	RefreshPhysicsAssetChange(PhysicsAsset);
 	return true;
 }
