Guard against invalid PhysicsAssetEditor selection crashing when calling ResetBoneCollision. #jira UE-67317 #rb none

diff --git a/Engine/Source/Editor/PhysicsAssetEditor/Private/PhysicsAssetEditor.cpp b/Engine/Source/Editor/PhysicsAssetEditor/Private/PhysicsAssetEditor.cpp
--- a/Engine/Source/Editor/PhysicsAssetEditor/Private/PhysicsAssetEditor.cpp
+++ b/Engine/Source/Editor/PhysicsAssetEditor/Private/PhysicsAssetEditor.cpp
@@ -1666,7 +1666,13 @@ void FPhysicsAssetEditor::ResetBoneCollision()
 		SlowTask.MakeDialog();
 		for(int32 i=0; i<SharedData->SelectedBodies.Num(); ++i)
 		{
-			UBodySetup* BodySetup = SharedData->PhysicsAsset->SkeletalBodySetups[SharedData->SelectedBodies[i].Index];
+			int32 SelectedBodyIndex = SharedData->SelectedBodies[i].Index;
+			if (SharedData->PhysicsAsset->SkeletalBodySetups.IsValidIndex(SelectedBodyIndex) == false)
+			{
+				continue;
+			}
+
+			UBodySetup* BodySetup = SharedData->PhysicsAsset->SkeletalBodySetups[SelectedBodyIndex];
 			check(BodySetup);
 			SlowTask.EnterProgressFrame(1.0f, FText::Format(LOCTEXT("ResetCollsionStepInfo", "Generating collision for {0}"), FText::FromName(BodySetup->BoneName)));
 			BodySetup->Modify();
@@ -1677,11 +1683,11 @@ void FPhysicsAssetEditor::ResetBoneCollision()
 			const FBoneVertInfo& UseVertInfo = NewBodyData.VertWeight == EVW_DominantWeight ? SharedData->DominantWeightBoneInfos[BoneIndex] : SharedData->AnyWeightBoneInfos[BoneIndex];
 			if(FPhysicsAssetUtils::CreateCollisionFromBone(BodySetup, EditorSkelMesh, BoneIndex, NewBodyData, UseVertInfo))
 			{
-				BodyIndices.AddUnique(SharedData->SelectedBodies[i].Index);
+				BodyIndices.AddUnique(SelectedBodyIndex);
 			}
 			else
 			{
-				FPhysicsAssetUtils::DestroyBody(SharedData->PhysicsAsset, SharedData->SelectedBodies[i].Index);
+				FPhysicsAssetUtils::DestroyBody(SharedData->PhysicsAsset, SelectedBodyIndex);
 			}
 		}
 
