Renamed command to relfect actual behavior.

diff --git a/Engine/Source/Editor/PhysicsAssetEditor/Private/PhysicsAssetEditorActions.cpp b/Engine/Source/Editor/PhysicsAssetEditor/Private/PhysicsAssetEditorActions.cpp
--- a/Engine/Source/Editor/PhysicsAssetEditor/Private/PhysicsAssetEditorActions.cpp
+++ b/Engine/Source/Editor/PhysicsAssetEditor/Private/PhysicsAssetEditorActions.cpp
@@ -53,7 +53,7 @@ void FPhysicsAssetEditorCommands::RegisterCommands()
 	UI_COMMAND(DeletePrimitive, "Delete", "Delete Selected Primitive(s)", EUserInterfaceActionType::Button, FInputChord());
 	UI_COMMAND(DuplicatePrimitive, "Duplicate", "Duplicate Selected Primitive(s)", EUserInterfaceActionType::Button, FInputChord());
 	UI_COMMAND(ResetConstraint, "Reset", "Reset Constraint", EUserInterfaceActionType::Button, FInputChord());
-	UI_COMMAND(SnapConstraint, "Snap", "Snap Constraint To Bone", EUserInterfaceActionType::Button, FInputChord());
+	UI_COMMAND(SnapConstraint, "Snap", "Snap Constraint Translation To Bone", EUserInterfaceActionType::Button, FInputChord());
 	UI_COMMAND(ConvertToBallAndSocket, "To Ball & Socket", "Convert Selected Constraint To Ball-And-Socket", EUserInterfaceActionType::Button, FInputChord());
 	UI_COMMAND(ConvertToHinge, "To Hinge", "Convert Selected Constraint To Hinge", EUserInterfaceActionType::Button, FInputChord());
 	UI_COMMAND(ConvertToPrismatic, "To Prismatic", "Convert Selected Constraint To Prismatic", EUserInterfaceActionType::Button, FInputChord());
