Force physics collision on physics asset editor skeletal mesh component so that simulation is enabled for newly generated bodies (affects only phat sim).

diff --git a/Engine/Source/Editor/PhysicsAssetEditor/Private/PhysicsAssetEditor.cpp b/Engine/Source/Editor/PhysicsAssetEditor/Private/PhysicsAssetEditor.cpp
--- a/Engine/Source/Editor/PhysicsAssetEditor/Private/PhysicsAssetEditor.cpp
+++ b/Engine/Source/Editor/PhysicsAssetEditor/Private/PhysicsAssetEditor.cpp
@@ -1355,6 +1355,10 @@ void FPhysicsAssetEditor::AddNewPrimitive(EAggCollisionShape::Type InPrimitiveTy
 {
 	check(!bCopySelected || SharedData->SelectedBodies.Num() == 1);	//we only support this for one selection
 	int32 NewPrimIndex = 0;
+	if (SharedData->EditorSkelComp)
+	{
+		SharedData->EditorSkelComp->SetCollisionEnabled(ECollisionEnabled::QueryAndPhysics);
+	}
 	TArray<FPhysicsAssetEditorSharedData::FSelection> NewSelection;
 	{
 		// Make sure rendering is done - so we are not changing data being used by collision drawing.
@@ -1651,7 +1655,10 @@ void FPhysicsAssetEditor::ResetBoneCollision()
 	{
 		return;
 	}
-
+	if (SharedData->EditorSkelComp)
+	{
+		SharedData->EditorSkelComp->SetCollisionEnabled(ECollisionEnabled::QueryAndPhysics);
+	}
 	// Make sure rendering is done - so we are not changing data being used by collision drawing.
 	FlushRenderingCommands();
 
diff --git a/Engine/Source/Editor/PhysicsAssetEditor/Private/PhysicsAssetEditorSharedData.cpp b/Engine/Source/Editor/PhysicsAssetEditor/Private/PhysicsAssetEditorSharedData.cpp
--- a/Engine/Source/Editor/PhysicsAssetEditor/Private/PhysicsAssetEditorSharedData.cpp
+++ b/Engine/Source/Editor/PhysicsAssetEditor/Private/PhysicsAssetEditorSharedData.cpp
@@ -988,7 +988,10 @@ void FPhysicsAssetEditorSharedData::MakeNewBody(int32 NewBoneIndex, bool bAutoSe
 	{
 		return;
 	}
-
+	if(EditorSkelComp)
+	{
+		EditorSkelComp->SetCollisionEnabled(ECollisionEnabled::QueryAndPhysics);
+	}
 	PhysicsAsset->Modify();
 
 	FName NewBoneName = EditorSkelMesh->RefSkeleton.GetBoneName(NewBoneIndex);
