Phat: Add option to hide simulated/kinematic bodies (i.e their shapes) in preview.

diff --git a/Engine/Source/Editor/PhysicsAssetEditor/Private/PhysicsAssetEditor.cpp b/Engine/Source/Editor/PhysicsAssetEditor/Private/PhysicsAssetEditor.cpp
--- a/Engine/Source/Editor/PhysicsAssetEditor/Private/PhysicsAssetEditor.cpp
+++ b/Engine/Source/Editor/PhysicsAssetEditor/Private/PhysicsAssetEditor.cpp
@@ -965,6 +965,18 @@ void FPhysicsAssetEditor::BindCommands()
 		FCanExecuteAction(),
 		FIsActionChecked::CreateSP(this, &FPhysicsAssetEditor::IsRenderingOnlySelectedSolid));
 
+	ViewportCommandList->MapAction(
+		Commands.HideSimulatedBodies,
+		FExecuteAction::CreateSP(this, &FPhysicsAssetEditor::ToggleHideSimulatedBodies),
+		FCanExecuteAction(),
+		FIsActionChecked::CreateSP(this, &FPhysicsAssetEditor::IsHidingSimulatedBodies));
+
+	ViewportCommandList->MapAction(
+		Commands.HideKinematicBodies,
+		FExecuteAction::CreateSP(this, &FPhysicsAssetEditor::ToggleHideKinematicBodies),
+		FCanExecuteAction(),
+		FIsActionChecked::CreateSP(this, &FPhysicsAssetEditor::IsHidingKinematicBodies));
+
 	ViewportCommandList->MapAction(
 		Commands.DrawConstraintsAsPoints,
 		FExecuteAction::CreateSP(this, &FPhysicsAssetEditor::ToggleDrawConstraintsAsPoints),
@@ -2031,11 +2043,33 @@ void FPhysicsAssetEditor::ToggleRenderOnlySelectedSolid()
 	SharedData->EditorOptions->SaveConfig();
 }
 
+void FPhysicsAssetEditor::ToggleHideSimulatedBodies()
+{
+	SharedData->EditorOptions->bHideSimulatedBodies = !SharedData->EditorOptions->bHideSimulatedBodies;
+	SharedData->EditorOptions->SaveConfig();
+}
+
+void FPhysicsAssetEditor::ToggleHideKinematicBodies()
+{
+	SharedData->EditorOptions->bHideKinematicBodies = !SharedData->EditorOptions->bHideKinematicBodies;
+	SharedData->EditorOptions->SaveConfig();
+}
+
 bool FPhysicsAssetEditor::IsRenderingOnlySelectedSolid() const
 {
 	return SharedData->EditorOptions->bSolidRenderingForSelectedOnly;
 }
 
+bool FPhysicsAssetEditor::IsHidingSimulatedBodies() const
+{
+	return SharedData->EditorOptions->bHideSimulatedBodies;
+}
+
+bool FPhysicsAssetEditor::IsHidingKinematicBodies() const
+{
+	return SharedData->EditorOptions->bHideKinematicBodies;
+}
+
 bool FPhysicsAssetEditor::IsConstraintRenderingMode(EPhysicsAssetEditorConstraintViewMode Mode, bool bSimulation) const
 {
 	return Mode == SharedData->GetCurrentConstraintViewMode(bSimulation);
diff --git a/Engine/Source/Editor/PhysicsAssetEditor/Private/PhysicsAssetEditor.h b/Engine/Source/Editor/PhysicsAssetEditor/Private/PhysicsAssetEditor.h
--- a/Engine/Source/Editor/PhysicsAssetEditor/Private/PhysicsAssetEditor.h
+++ b/Engine/Source/Editor/PhysicsAssetEditor/Private/PhysicsAssetEditor.h
@@ -222,7 +222,11 @@ private:
 	void ToggleRenderOnlySelectedConstraints();
 	bool IsRenderingOnlySelectedConstraints() const;
 	void ToggleRenderOnlySelectedSolid();
+	void ToggleHideSimulatedBodies();
+	void ToggleHideKinematicBodies();
 	bool IsRenderingOnlySelectedSolid() const;
+	bool IsHidingSimulatedBodies() const;
+	bool IsHidingKinematicBodies() const;
 	void OnToggleMassProperties();
 	bool IsToggleMassProperties() const;
 	void OnSetCollision(bool bEnable);
diff --git a/Engine/Source/Editor/PhysicsAssetEditor/Private/PhysicsAssetEditorActions.cpp b/Engine/Source/Editor/PhysicsAssetEditor/Private/PhysicsAssetEditorActions.cpp
--- a/Engine/Source/Editor/PhysicsAssetEditor/Private/PhysicsAssetEditorActions.cpp
+++ b/Engine/Source/Editor/PhysicsAssetEditor/Private/PhysicsAssetEditorActions.cpp
@@ -38,6 +38,8 @@ void FPhysicsAssetEditorCommands::RegisterCommands()
 	UI_COMMAND(ConstraintRenderingMode_Simulation_AllLimits, "All Limits", "All Limits Constraint Rendering Mode (Simulation)", EUserInterfaceActionType::RadioButton, FInputChord());
 
 	UI_COMMAND(RenderOnlySelectedSolid, "Only Selected Solid", "Only render selected collision as 'solid'", EUserInterfaceActionType::ToggleButton, FInputChord());
+	UI_COMMAND(HideSimulatedBodies, "Hide Simulated Bodies", "Hide rendering for simulated bodies", EUserInterfaceActionType::ToggleButton, FInputChord());
+	UI_COMMAND(HideKinematicBodies, "Hide Kinematic Bodies", "Hide rendering for kinematic bodies", EUserInterfaceActionType::ToggleButton, FInputChord());
 	UI_COMMAND(RenderOnlySelectedConstraints, "Only Selected Constraints", "Draw only selected constraints.", EUserInterfaceActionType::ToggleButton, FInputChord());
 	UI_COMMAND(DrawConstraintsAsPoints, "Draw Constraints As Points", "Draw Constraints As Points", EUserInterfaceActionType::ToggleButton, FInputChord());
 	UI_COMMAND(ToggleMassProperties, "Mass Properties", "Show Mass Properties For Bodies When Simulation Is Enabled", EUserInterfaceActionType::ToggleButton, FInputChord());
diff --git a/Engine/Source/Editor/PhysicsAssetEditor/Private/PhysicsAssetEditorActions.h b/Engine/Source/Editor/PhysicsAssetEditor/Private/PhysicsAssetEditorActions.h
--- a/Engine/Source/Editor/PhysicsAssetEditor/Private/PhysicsAssetEditorActions.h
+++ b/Engine/Source/Editor/PhysicsAssetEditor/Private/PhysicsAssetEditorActions.h
@@ -50,6 +50,8 @@ public:
 	TSharedPtr<FUICommandInfo> ConstraintRenderingMode_Simulation_AllPositions;
 	TSharedPtr<FUICommandInfo> ConstraintRenderingMode_Simulation_AllLimits;
 	TSharedPtr<FUICommandInfo> RenderOnlySelectedSolid;
+	TSharedPtr<FUICommandInfo> HideSimulatedBodies;
+	TSharedPtr<FUICommandInfo> HideKinematicBodies;
 	TSharedPtr<FUICommandInfo> RenderOnlySelectedConstraints;
 	TSharedPtr<FUICommandInfo> DrawConstraintsAsPoints;
 	TSharedPtr<FUICommandInfo> ToggleMassProperties;
diff --git a/Engine/Source/Editor/PhysicsAssetEditor/Private/PhysicsAssetEditorMode.cpp b/Engine/Source/Editor/PhysicsAssetEditor/Private/PhysicsAssetEditorMode.cpp
--- a/Engine/Source/Editor/PhysicsAssetEditor/Private/PhysicsAssetEditorMode.cpp
+++ b/Engine/Source/Editor/PhysicsAssetEditor/Private/PhysicsAssetEditorMode.cpp
@@ -84,6 +84,8 @@ FPhysicsAssetEditorMode::FPhysicsAssetEditorMode(TSharedRef<FWorkflowCentricAppl
 					InSubMenuBuilder.BeginSection("PhysicsAssetEditorCollisionRenderSettings", LOCTEXT("CollisionRenderSettingsHeader", "Body Drawing"));
 					{
 						InSubMenuBuilder.AddMenuEntry(Commands.RenderOnlySelectedSolid);
+						InSubMenuBuilder.AddMenuEntry(Commands.HideSimulatedBodies);
+						InSubMenuBuilder.AddMenuEntry(Commands.HideKinematicBodies);
 						InSubMenuBuilder.AddWidget(PhysicsAssetEditorPtr.Pin()->MakeCollisionOpacityWidget(), LOCTEXT("CollisionOpacityLabel", "Collision Opacity"));
 					}
 					InSubMenuBuilder.EndSection();
diff --git a/Engine/Source/Editor/PhysicsAssetEditor/Private/PhysicsAssetEditorSkeletalMeshComponent.cpp b/Engine/Source/Editor/PhysicsAssetEditor/Private/PhysicsAssetEditorSkeletalMeshComponent.cpp
--- a/Engine/Source/Editor/PhysicsAssetEditor/Private/PhysicsAssetEditorSkeletalMeshComponent.cpp
+++ b/Engine/Source/Editor/PhysicsAssetEditor/Private/PhysicsAssetEditorSkeletalMeshComponent.cpp
@@ -108,6 +108,14 @@ void UPhysicsAssetEditorSkeletalMeshComponent::RenderAssetTools(const FSceneView
 		{
 			continue;
 		}
+		if ((PhysicsAsset->SkeletalBodySetups[i]->PhysicsType == EPhysicsType::PhysType_Kinematic &&
+			SharedData->EditorOptions->bHideKinematicBodies) ||
+			(PhysicsAsset->SkeletalBodySetups[i]->PhysicsType == EPhysicsType::PhysType_Simulated &&
+			SharedData->EditorOptions->bHideSimulatedBodies)
+			)
+		{
+			continue;
+		}
 		int32 BoneIndex = GetBoneIndex(PhysicsAsset->SkeletalBodySetups[i]->BoneName);
 
 		// If we found a bone for it, draw the collision.
diff --git a/Engine/Source/Editor/UnrealEd/Classes/Preferences/PhysicsAssetEditorOptions.h b/Engine/Source/Editor/UnrealEd/Classes/Preferences/PhysicsAssetEditorOptions.h
--- a/Engine/Source/Editor/UnrealEd/Classes/Preferences/PhysicsAssetEditorOptions.h
+++ b/Engine/Source/Editor/UnrealEd/Classes/Preferences/PhysicsAssetEditorOptions.h
@@ -134,6 +134,14 @@ class UNREALED_API UPhysicsAssetEditorOptions : public UObject
 	UPROPERTY(config)
 	bool bSolidRenderingForSelectedOnly;
 
+	/** When set, disables rendering for simulated bodies */
+	UPROPERTY(config)
+	bool bHideSimulatedBodies;
+
+	/** When set, disables rendering for kinematic bodies */
+	UPROPERTY(config)
+	bool bHideKinematicBodies;
+
 	/** When set, cloth will reset each time simulation is toggled */
 	UPROPERTY(EditAnywhere, config, Category=Clothing)
 	bool bResetClothWhenSimulating;
diff --git a/Engine/Source/Editor/UnrealEd/Private/PreferenceStubs.cpp b/Engine/Source/Editor/UnrealEd/Private/PreferenceStubs.cpp
--- a/Engine/Source/Editor/UnrealEd/Private/PreferenceStubs.cpp
+++ b/Engine/Source/Editor/UnrealEd/Private/PreferenceStubs.cpp
@@ -52,6 +52,8 @@ UPhysicsAssetEditorOptions::UPhysicsAssetEditorOptions(const FObjectInitializer&
 
 	CollisionOpacity = 0.3f;
 	bSolidRenderingForSelectedOnly = false;
+	bHideSimulatedBodies = false;
+	bHideKinematicBodies = false;
 	bResetClothWhenSimulating = false;
 }
 
