Phat:
Select Simulated/Kinematic bodies.
Show/Hide Simulated/Kinematic bodies in Skeleton Tree.

diff --git a/Engine/Source/Editor/PhysicsAssetEditor/Private/PhysicsAssetEditor.cpp b/Engine/Source/Editor/PhysicsAssetEditor/Private/PhysicsAssetEditor.cpp
--- a/Engine/Source/Editor/PhysicsAssetEditor/Private/PhysicsAssetEditor.cpp
+++ b/Engine/Source/Editor/PhysicsAssetEditor/Private/PhysicsAssetEditor.cpp
@@ -565,6 +565,8 @@ void FPhysicsAssetEditor::ExtendMenu()
 			const FPhysicsAssetEditorCommands& Commands = FPhysicsAssetEditorCommands::Get();
 			MenuBarBuilder.BeginSection("Selection", LOCTEXT("PhatEditSelection", "Selection"));
 			MenuBarBuilder.AddMenuEntry(Commands.SelectAllBodies);
+			MenuBarBuilder.AddMenuEntry(Commands.SelectSimulatedBodies);
+			MenuBarBuilder.AddMenuEntry(Commands.SelectKinematicBodies);
 			MenuBarBuilder.AddMenuEntry(Commands.SelectAllConstraints);
 			MenuBarBuilder.AddMenuEntry(Commands.ToggleSelectionType);
 			MenuBarBuilder.AddMenuEntry(Commands.DeselectAll);
@@ -804,6 +806,16 @@ void FPhysicsAssetEditor::BindCommands()
 		FExecuteAction::CreateSP(this, &FPhysicsAssetEditor::OnSelectAllBodies),
 		FCanExecuteAction::CreateSP(this, &FPhysicsAssetEditor::IsNotSimulation));
 
+	ToolkitCommands->MapAction(
+		Commands.SelectSimulatedBodies,
+		FExecuteAction::CreateSP(this, &FPhysicsAssetEditor::OnSelectSimulatedBodies),
+		FCanExecuteAction::CreateSP(this, &FPhysicsAssetEditor::IsNotSimulation));
+
+	ToolkitCommands->MapAction(
+		Commands.SelectKinematicBodies,
+		FExecuteAction::CreateSP(this, &FPhysicsAssetEditor::OnSelectKinematicBodies),
+		FCanExecuteAction::CreateSP(this, &FPhysicsAssetEditor::IsNotSimulation));
+
 	ToolkitCommands->MapAction(
 		Commands.SelectAllConstraints,
 		FExecuteAction::CreateSP(this, &FPhysicsAssetEditor::OnSelectAllConstraints),
@@ -1004,6 +1016,20 @@ void FPhysicsAssetEditor::BindCommands()
 		FGetActionCheckState::CreateSP(this, &FPhysicsAssetEditor::GetShowBodiesChecked)
 	);
 
+	SkeletonTreeCommandList->MapAction(
+		Commands.ShowSimulatedBodies,
+		FExecuteAction::CreateSP(this, &FPhysicsAssetEditor::HandleToggleShowSimulatedBodies),
+		FCanExecuteAction(),
+		FGetActionCheckState::CreateSP(this, &FPhysicsAssetEditor::GetShowSimulatedBodiesChecked)
+	);
+
+	SkeletonTreeCommandList->MapAction(
+		Commands.ShowKinematicBodies,
+		FExecuteAction::CreateSP(this, &FPhysicsAssetEditor::HandleToggleShowKinematicBodies),
+		FCanExecuteAction(),
+		FGetActionCheckState::CreateSP(this, &FPhysicsAssetEditor::GetShowKinematicBodiesChecked)
+	);
+
 	SkeletonTreeCommandList->MapAction(
 		Commands.ShowConstraints,
 		FExecuteAction::CreateSP(this, &FPhysicsAssetEditor::HandleToggleShowConstraints),
@@ -1214,6 +1240,8 @@ void FPhysicsAssetEditor::BuildMenuWidgetSelection(FMenuBuilder& InMenuBuilder)
 
 		InMenuBuilder.BeginSection( "Selection", LOCTEXT("Selection", "Selection" ) );
 		InMenuBuilder.AddMenuEntry( Commands.SelectAllBodies );
+		InMenuBuilder.AddMenuEntry(Commands.SelectSimulatedBodies);
+		InMenuBuilder.AddMenuEntry(Commands.SelectKinematicBodies);
 		InMenuBuilder.AddMenuEntry( Commands.SelectAllConstraints );
 		InMenuBuilder.AddMenuEntry( Commands.ToggleSelectionType );
 		InMenuBuilder.EndSection();
@@ -1232,6 +1260,8 @@ TSharedRef<ISkeletonTree> FPhysicsAssetEditor::BuildMenuWidgetNewConstraintForBo
 
 	TSharedRef<FPhysicsAssetEditorSkeletonTreeBuilder> Builder = MakeShared<FPhysicsAssetEditorSkeletonTreeBuilder>(SharedData->PhysicsAsset, SkeletonTreeBuilderArgs);
 	Builder->bShowBodies = true;
+	Builder->bShowSimulatedBodies = true;
+	Builder->bShowKinematicBodies = true;
 	Builder->bShowConstraints = false;
 	Builder->bShowPrimitives = false;
 
@@ -2573,6 +2603,70 @@ void FPhysicsAssetEditor::OnSelectAllBodies()
 	}
 }
 
+void FPhysicsAssetEditor::OnSelectKinematicBodies()
+{
+	OnSelectBodies(EPhysicsType::PhysType_Kinematic);
+}
+
+void FPhysicsAssetEditor::OnSelectSimulatedBodies()
+{
+	OnSelectBodies(EPhysicsType::PhysType_Simulated);
+}
+
+void FPhysicsAssetEditor::OnSelectBodies(EPhysicsType PhysicsType)
+{
+	UPhysicsAsset * const PhysicsAsset = SharedData->EditorSkelComp->GetPhysicsAsset();
+
+	// Block selection broadcast until we have selected all, as this can be an expensive operation
+	FScopedBulkSelection BulkSelection(SharedData);
+
+	//Bodies
+	//first deselect everything
+	SharedData->ClearSelectedBody();
+
+	//go through every body and add every geom
+	for (int32 i = 0; i < PhysicsAsset->SkeletalBodySetups.Num(); ++i)
+	{
+		int32 BoneIndex = SharedData->EditorSkelComp->GetBoneIndex(PhysicsAsset->SkeletalBodySetups[i]->BoneName);
+		if (PhysicsAsset->SkeletalBodySetups[i]->PhysicsType == PhysicsType)
+		{
+			// If we found a bone for it, add all geom
+			if (BoneIndex != INDEX_NONE)
+			{
+				FKAggregateGeom* AggGeom = &PhysicsAsset->SkeletalBodySetups[i]->AggGeom;
+				for (int32 j = 0; j < AggGeom->SphereElems.Num(); ++j)
+				{
+					FPhysicsAssetEditorSharedData::FSelection Selection(i, EAggCollisionShape::Sphere, j);
+					SharedData->SetSelectedBody(Selection, true);
+				}
+
+				for (int32 j = 0; j < AggGeom->BoxElems.Num(); ++j)
+				{
+					FPhysicsAssetEditorSharedData::FSelection Selection(i, EAggCollisionShape::Box, j);
+					SharedData->SetSelectedBody(Selection, true);
+				}
+
+				for (int32 j = 0; j < AggGeom->SphylElems.Num(); ++j)
+				{
+					FPhysicsAssetEditorSharedData::FSelection Selection(i, EAggCollisionShape::Sphyl, j);
+					SharedData->SetSelectedBody(Selection, true);
+				}
+
+				for (int32 j = 0; j < AggGeom->ConvexElems.Num(); ++j)
+				{
+					FPhysicsAssetEditorSharedData::FSelection Selection(i, EAggCollisionShape::Convex, j);
+					SharedData->SetSelectedBody(Selection, true);
+				}
+				for (int32 j = 0; j < AggGeom->TaperedCapsuleElems.Num(); ++j)
+				{
+					FPhysicsAssetEditorSharedData::FSelection Selection(i, EAggCollisionShape::TaperedCapsule, j);
+					SharedData->SetSelectedBody(Selection, true);
+				}
+			}
+		}
+	}
+}
+
 void FPhysicsAssetEditor::OnSelectAllConstraints()
 {
 	UPhysicsAsset * const PhysicsAsset = SharedData->EditorSkelComp->GetPhysicsAsset();
@@ -2889,6 +2983,8 @@ void FPhysicsAssetEditor::HandleExtendFilterMenu(FMenuBuilder& InMenuBuilder)
 	InMenuBuilder.BeginSection(TEXT("PhysicsAssetFilters"), LOCTEXT("PhysicsAssetFiltersHeader", "Physics Asset Filters"));
 	{
 		InMenuBuilder.AddMenuEntry(Commands.ShowBodies);
+		InMenuBuilder.AddMenuEntry(Commands.ShowSimulatedBodies);
+		InMenuBuilder.AddMenuEntry(Commands.ShowKinematicBodies);
 		InMenuBuilder.AddMenuEntry(Commands.ShowConstraints);
 		InMenuBuilder.AddMenuEntry(Commands.ShowPrimitives);
 	}
@@ -2902,6 +2998,18 @@ void FPhysicsAssetEditor::HandleToggleShowBodies()
 	SkeletonTree->RefreshFilter();
 }
 
+void FPhysicsAssetEditor::HandleToggleShowSimulatedBodies()
+{
+	SkeletonTreeBuilder->bShowSimulatedBodies = !SkeletonTreeBuilder->bShowSimulatedBodies;
+	SkeletonTree->RefreshFilter();
+}
+
+void FPhysicsAssetEditor::HandleToggleShowKinematicBodies()
+{
+	SkeletonTreeBuilder->bShowKinematicBodies = !SkeletonTreeBuilder->bShowKinematicBodies;
+	SkeletonTree->RefreshFilter();
+}
+
 void FPhysicsAssetEditor::HandleToggleShowConstraints()
 {
 	SkeletonTreeBuilder->bShowConstraints = !SkeletonTreeBuilder->bShowConstraints;
@@ -2919,6 +3027,16 @@ ECheckBoxState FPhysicsAssetEditor::GetShowBodiesChecked() const
 	return SkeletonTreeBuilder->bShowBodies ? ECheckBoxState::Checked : ECheckBoxState::Unchecked;
 }
 
+ECheckBoxState FPhysicsAssetEditor::GetShowSimulatedBodiesChecked() const
+{
+	return SkeletonTreeBuilder->bShowSimulatedBodies ? ECheckBoxState::Checked : ECheckBoxState::Unchecked;
+}
+
+ECheckBoxState FPhysicsAssetEditor::GetShowKinematicBodiesChecked() const
+{
+	return SkeletonTreeBuilder->bShowKinematicBodies ? ECheckBoxState::Checked : ECheckBoxState::Unchecked;
+}
+
 ECheckBoxState FPhysicsAssetEditor::GetShowConstraintsChecked() const
 {
 	return SkeletonTreeBuilder->bShowConstraints ? ECheckBoxState::Checked : ECheckBoxState::Unchecked;
diff --git a/Engine/Source/Editor/PhysicsAssetEditor/Private/PhysicsAssetEditor.h b/Engine/Source/Editor/PhysicsAssetEditor/Private/PhysicsAssetEditor.h
--- a/Engine/Source/Editor/PhysicsAssetEditor/Private/PhysicsAssetEditor.h
+++ b/Engine/Source/Editor/PhysicsAssetEditor/Private/PhysicsAssetEditor.h
@@ -269,6 +269,9 @@ private:
 
 	//menu commands
 	void OnSelectAllBodies();
+	void OnSelectKinematicBodies();
+	void OnSelectSimulatedBodies();
+	void OnSelectBodies(EPhysicsType PhysicsType = EPhysicsType::PhysType_Simulated);
 	void OnSelectAllConstraints();
 	void OnToggleSelectionType();
 	void OnDeselectAll();
@@ -287,9 +290,13 @@ private:
 
 	/** Filter menu toggles */
 	void HandleToggleShowBodies();
+	void HandleToggleShowSimulatedBodies();
+	void HandleToggleShowKinematicBodies();
 	void HandleToggleShowConstraints();
 	void HandleToggleShowPrimitives();
 	ECheckBoxState GetShowBodiesChecked() const;
+	ECheckBoxState GetShowSimulatedBodiesChecked() const;
+	ECheckBoxState GetShowKinematicBodiesChecked() const;
 	ECheckBoxState GetShowConstraintsChecked() const;
 	ECheckBoxState GetShowPrimitivesChecked() const;
 
diff --git a/Engine/Source/Editor/PhysicsAssetEditor/Private/PhysicsAssetEditorActions.cpp b/Engine/Source/Editor/PhysicsAssetEditor/Private/PhysicsAssetEditorActions.cpp
--- a/Engine/Source/Editor/PhysicsAssetEditor/Private/PhysicsAssetEditorActions.cpp
+++ b/Engine/Source/Editor/PhysicsAssetEditor/Private/PhysicsAssetEditorActions.cpp
@@ -71,6 +71,8 @@ void FPhysicsAssetEditorCommands::RegisterCommands()
 	UI_COMMAND(DeleteBody, "Delete", "Delete Selected Body", EUserInterfaceActionType::Button, FInputChord());
 	UI_COMMAND(DeleteAllBodiesBelow, "Delete All Bodies Below", "Delete All Bodies Below", EUserInterfaceActionType::Button, FInputChord());
 	UI_COMMAND(SelectAllBodies, "Select All Bodies", "Select All Bodies", EUserInterfaceActionType::Button, FInputChord(EModifierKey::Control, EKeys::A));
+	UI_COMMAND(SelectSimulatedBodies, "Select Simulated Bodies", "Select Simulated Bodies", EUserInterfaceActionType::Button, FInputChord(EModifierKey::Control, EKeys::J));
+	UI_COMMAND(SelectKinematicBodies, "Select Kinematic Bodies", "Select Kinematic Bodies", EUserInterfaceActionType::Button, FInputChord(EModifierKey::Control, EKeys::K));
 	UI_COMMAND(SelectAllConstraints, "Select All Constraints", "Select All Constraints", EUserInterfaceActionType::Button, FInputChord(EModifierKey::Control | EModifierKey::Shift, EKeys::A));
 	UI_COMMAND(ToggleSelectionType, "Toggle Selection Type", "Select bodies from constraints, or constraints from bodies, depending on selection.", EUserInterfaceActionType::Button, FInputChord(EModifierKey::Control, EKeys::T));
 	UI_COMMAND(DeselectAll, "Deselect All Bodies", "Deselect All", EUserInterfaceActionType::Button, FInputChord(EKeys::Escape));
@@ -88,6 +90,8 @@ void FPhysicsAssetEditorCommands::RegisterCommands()
 	UI_COMMAND(RemoveConstraintFromCurrentConstraintProfile, "Unassign", "Unassigns the selected constraints from the current constraint profile", EUserInterfaceActionType::Button, FInputChord());
 
 	UI_COMMAND(ShowBodies, "Show Bodies", "Display bodies in the tree view. Bodies are a collection of primitive shapes used for collision.", EUserInterfaceActionType::ToggleButton, FInputChord());
+	UI_COMMAND(ShowSimulatedBodies, "Show Simulated Bodies", "Display simulated bodies in the tree view. Bodies are a collection of primitive shapes used for collision.", EUserInterfaceActionType::ToggleButton, FInputChord());
+	UI_COMMAND(ShowKinematicBodies, "Show Kinematic Bodies", "Display kinematic bodies in the tree view. Bodies are a collection of primitive shapes used for collision.", EUserInterfaceActionType::ToggleButton, FInputChord());
 	UI_COMMAND(ShowConstraints, "Show Constraints", "Display constraints in the tree view. Constraints are used to control how bodies can move in relation to one another.", EUserInterfaceActionType::ToggleButton, FInputChord());
 	UI_COMMAND(ShowPrimitives, "Show Primitives", "Display primitive shapes (sphere, boxes, capsules etc.) in the tree view", EUserInterfaceActionType::ToggleButton, FInputChord());
 
diff --git a/Engine/Source/Editor/PhysicsAssetEditor/Private/PhysicsAssetEditorActions.h b/Engine/Source/Editor/PhysicsAssetEditor/Private/PhysicsAssetEditorActions.h
--- a/Engine/Source/Editor/PhysicsAssetEditor/Private/PhysicsAssetEditorActions.h
+++ b/Engine/Source/Editor/PhysicsAssetEditor/Private/PhysicsAssetEditorActions.h
@@ -83,6 +83,8 @@ public:
 	TSharedPtr<FUICommandInfo> DeleteBody;
 	TSharedPtr<FUICommandInfo> DeleteAllBodiesBelow;
 	TSharedPtr<FUICommandInfo> SelectAllBodies;
+	TSharedPtr<FUICommandInfo> SelectSimulatedBodies;
+	TSharedPtr<FUICommandInfo> SelectKinematicBodies;
 	TSharedPtr<FUICommandInfo> SelectAllConstraints;
 	TSharedPtr<FUICommandInfo> ToggleSelectionType;
 	TSharedPtr<FUICommandInfo> DeselectAll;
@@ -98,6 +100,8 @@ public:
 	TSharedPtr<FUICommandInfo> AddConstraintToCurrentConstraintProfile;
 	TSharedPtr<FUICommandInfo> RemoveConstraintFromCurrentConstraintProfile;
 	TSharedPtr<FUICommandInfo> ShowBodies;
+	TSharedPtr<FUICommandInfo> ShowSimulatedBodies;
+	TSharedPtr<FUICommandInfo> ShowKinematicBodies;
 	TSharedPtr<FUICommandInfo> ShowConstraints;
 	TSharedPtr<FUICommandInfo> ShowPrimitives;
 
diff --git a/Engine/Source/Editor/PhysicsAssetEditor/Private/PhysicsAssetEditorSkeletonTreeBuilder.cpp b/Engine/Source/Editor/PhysicsAssetEditor/Private/PhysicsAssetEditorSkeletonTreeBuilder.cpp
--- a/Engine/Source/Editor/PhysicsAssetEditor/Private/PhysicsAssetEditorSkeletonTreeBuilder.cpp
+++ b/Engine/Source/Editor/PhysicsAssetEditor/Private/PhysicsAssetEditorSkeletonTreeBuilder.cpp
@@ -55,7 +55,26 @@ ESkeletonTreeFilterResult FPhysicsAssetEditorSkeletonTreeBuilder::FilterItem(con
 
 		if(InItem->IsOfType<FSkeletonTreePhysicsBodyItem>())
 		{
+			bool bShouldHideBody = false;
 			if(!bShowBodies)
+			{
+				bShouldHideBody = true;
+			}
+			else
+			{
+				if (UBodySetup* BodySetup = Cast<UBodySetup>(InItem->GetObject()))
+				{
+					if (BodySetup->PhysicsType == EPhysicsType::PhysType_Simulated && !bShowSimulatedBodies)
+					{
+						bShouldHideBody = true;
+					}
+					else if (BodySetup->PhysicsType == EPhysicsType::PhysType_Kinematic && !bShowKinematicBodies)
+					{
+						bShouldHideBody = true;
+					}
+				}
+			}
+			if (bShouldHideBody)
 			{
 				Result = ESkeletonTreeFilterResult::Hidden;
 			}
diff --git a/Engine/Source/Editor/PhysicsAssetEditor/Private/PhysicsAssetEditorSkeletonTreeBuilder.h b/Engine/Source/Editor/PhysicsAssetEditor/Private/PhysicsAssetEditorSkeletonTreeBuilder.h
--- a/Engine/Source/Editor/PhysicsAssetEditor/Private/PhysicsAssetEditorSkeletonTreeBuilder.h
+++ b/Engine/Source/Editor/PhysicsAssetEditor/Private/PhysicsAssetEditorSkeletonTreeBuilder.h
@@ -18,6 +18,8 @@ public:
 	
 	/** Flags used for filtering */
 	bool bShowBodies;
+	bool bShowKinematicBodies;
+	bool bShowSimulatedBodies;
 	bool bShowConstraints;
 	bool bShowPrimitives;
 
