Fix simulate preview not playing in phat (stop allowing AnimationEditorViewportClient to reset skel mesh comp collision to QueryOnly via setting profile).

diff --git a/Engine/Source/Editor/Persona/Private/AnimationEditorViewportClient.cpp b/Engine/Source/Editor/Persona/Private/AnimationEditorViewportClient.cpp
--- a/Engine/Source/Editor/Persona/Private/AnimationEditorViewportClient.cpp
+++ b/Engine/Source/Editor/Persona/Private/AnimationEditorViewportClient.cpp
@@ -379,10 +379,12 @@ void FAnimationViewportClient::HandleSkeletalMeshChanged(USkeletalMesh* OldSkele
 		PhysAsset->InvalidateAllPhysicsMeshes();
 		PreviewMeshComponent->TermArticulated();
 		PreviewMeshComponent->InitArticulated(GetWorld()->GetPhysicsScene());
-
-		// Set to PhysicsActor to enable tracing regardless of project overrides
-		static FName CollisionProfileName(TEXT("PhysicsActor"));
-		PreviewMeshComponent->SetCollisionProfileName(CollisionProfileName);
+		if (PreviewMeshComponent->CanOverrideCollisionProfile())
+		{
+			// Set to PhysicsActor to enable tracing regardless of project overrides
+			static FName CollisionProfileName(TEXT("PhysicsActor"));
+			PreviewMeshComponent->SetCollisionProfileName(CollisionProfileName);
+		}
 	}
 
 	Invalidate();
diff --git a/Engine/Source/Editor/PhysicsAssetEditor/Private/PhysicsAssetEditor.cpp b/Engine/Source/Editor/PhysicsAssetEditor/Private/PhysicsAssetEditor.cpp
--- a/Engine/Source/Editor/PhysicsAssetEditor/Private/PhysicsAssetEditor.cpp
+++ b/Engine/Source/Editor/PhysicsAssetEditor/Private/PhysicsAssetEditor.cpp
@@ -1397,10 +1397,6 @@ void FPhysicsAssetEditor::AddNewPrimitive(EAggCollisionShape::Type InPrimitiveTy
 {
 	check(!bCopySelected || SharedData->SelectedBodies.Num() == 1);	//we only support this for one selection
 	int32 NewPrimIndex = 0;
-	if (SharedData->EditorSkelComp)
-	{
-		SharedData->EditorSkelComp->SetCollisionEnabled(ECollisionEnabled::QueryAndPhysics);
-	}
 	TArray<FPhysicsAssetEditorSharedData::FSelection> NewSelection;
 	{
 		// Make sure rendering is done - so we are not changing data being used by collision drawing.
@@ -1697,10 +1693,6 @@ void FPhysicsAssetEditor::ResetBoneCollision()
 	{
 		return;
 	}
-	if (SharedData->EditorSkelComp)
-	{
-		SharedData->EditorSkelComp->SetCollisionEnabled(ECollisionEnabled::QueryAndPhysics);
-	}
 	// Make sure rendering is done - so we are not changing data being used by collision drawing.
 	FlushRenderingCommands();
 
diff --git a/Engine/Source/Editor/PhysicsAssetEditor/Private/PhysicsAssetEditorSharedData.cpp b/Engine/Source/Editor/PhysicsAssetEditor/Private/PhysicsAssetEditorSharedData.cpp
--- a/Engine/Source/Editor/PhysicsAssetEditor/Private/PhysicsAssetEditorSharedData.cpp
+++ b/Engine/Source/Editor/PhysicsAssetEditor/Private/PhysicsAssetEditorSharedData.cpp
@@ -988,10 +988,6 @@ void FPhysicsAssetEditorSharedData::MakeNewBody(int32 NewBoneIndex, bool bAutoSe
 	{
 		return;
 	}
-	if(EditorSkelComp)
-	{
-		EditorSkelComp->SetCollisionEnabled(ECollisionEnabled::QueryAndPhysics);
-	}
 	PhysicsAsset->Modify();
 
 	FName NewBoneName = EditorSkelMesh->RefSkeleton.GetBoneName(NewBoneIndex);
diff --git a/Engine/Source/Editor/PhysicsAssetEditor/Private/PhysicsAssetEditorSkeletalMeshComponent.h b/Engine/Source/Editor/PhysicsAssetEditor/Private/PhysicsAssetEditorSkeletalMeshComponent.h
--- a/Engine/Source/Editor/PhysicsAssetEditor/Private/PhysicsAssetEditorSkeletalMeshComponent.h
+++ b/Engine/Source/Editor/PhysicsAssetEditor/Private/PhysicsAssetEditorSkeletalMeshComponent.h
@@ -69,4 +69,6 @@ class UPhysicsAssetEditorSkeletalMeshComponent : public UDebugSkelMeshComponent
 	FTransform GetPrimitiveTransform(FTransform& BoneTM, int32 BodyIndex, EAggCollisionShape::Type PrimType, int32 PrimIndex, float Scale);
 	FColor GetPrimitiveColor(int32 BodyIndex, EAggCollisionShape::Type PrimitiveType, int32 PrimitiveIndex);
 	UMaterialInterface* GetPrimitiveMaterial(int32 BodyIndex, EAggCollisionShape::Type PrimitiveType, int32 PrimitiveIndex);
+public:
+	virtual bool CanOverrideCollisionProfile() const override { return false;  }
 };
diff --git a/Engine/Source/Editor/UnrealEd/Classes/Animation/DebugSkelMeshComponent.h b/Engine/Source/Editor/UnrealEd/Classes/Animation/DebugSkelMeshComponent.h
--- a/Engine/Source/Editor/UnrealEd/Classes/Animation/DebugSkelMeshComponent.h
+++ b/Engine/Source/Editor/UnrealEd/Classes/Animation/DebugSkelMeshComponent.h
@@ -472,6 +472,8 @@ public:
 
 	void RefreshSelectedClothingSkinnedPositions();
 
+	virtual bool CanOverrideCollisionProfile() const { return true; }
+
 	virtual void GetUsedMaterials(TArray<UMaterialInterface *>& OutMaterials, bool bGetDebugMaterials = false) const override;
 
 	/**
