Initialize show/hide bools for kinematic/simulated bodes in phat skeleton tree.

#rnx

diff --git a/Engine/Source/Editor/PhysicsAssetEditor/Private/PhysicsAssetEditorSkeletonTreeBuilder.cpp b/Engine/Source/Editor/PhysicsAssetEditor/Private/PhysicsAssetEditorSkeletonTreeBuilder.cpp
--- a/Engine/Source/Editor/PhysicsAssetEditor/Private/PhysicsAssetEditorSkeletonTreeBuilder.cpp
+++ b/Engine/Source/Editor/PhysicsAssetEditor/Private/PhysicsAssetEditorSkeletonTreeBuilder.cpp
@@ -14,6 +14,8 @@
 FPhysicsAssetEditorSkeletonTreeBuilder::FPhysicsAssetEditorSkeletonTreeBuilder(UPhysicsAsset* InPhysicsAsset, const FSkeletonTreeBuilderArgs& InSkeletonTreeBuilderArgs)
 	: FSkeletonTreeBuilder(InSkeletonTreeBuilderArgs)
 	, bShowBodies(true)
+	, bShowKinematicBodies(true)
+	, bShowSimulatedBodies(true)
 	, bShowConstraints(true)
 	, bShowPrimitives(true)
 	, PhysicsAsset(InPhysicsAsset)
