Option to show/hide selected bodies/constraints in phat.

diff --git a/Engine/Source/Editor/PhysicsAssetEditor/Private/PhysicsAssetEditor.cpp b/Engine/Source/Editor/PhysicsAssetEditor/Private/PhysicsAssetEditor.cpp
--- a/Engine/Source/Editor/PhysicsAssetEditor/Private/PhysicsAssetEditor.cpp
+++ b/Engine/Source/Editor/PhysicsAssetEditor/Private/PhysicsAssetEditor.cpp
@@ -569,6 +569,7 @@ void FPhysicsAssetEditor::ExtendMenu()
 			MenuBarBuilder.AddMenuEntry(Commands.SelectKinematicBodies);
 			MenuBarBuilder.AddMenuEntry(Commands.SelectAllConstraints);
 			MenuBarBuilder.AddMenuEntry(Commands.ToggleSelectionType);
+			MenuBarBuilder.AddMenuEntry(Commands.ToggleShowSelected);
 			MenuBarBuilder.AddMenuEntry(Commands.DeselectAll);
 			MenuBarBuilder.EndSection();
 		}
@@ -826,6 +827,11 @@ void FPhysicsAssetEditor::BindCommands()
 		FExecuteAction::CreateSP(this, &FPhysicsAssetEditor::OnToggleSelectionType),
 		FCanExecuteAction::CreateSP(this, &FPhysicsAssetEditor::IsNotSimulation));
 
+	ToolkitCommands->MapAction(
+		Commands.ToggleShowSelected,
+		FExecuteAction::CreateSP(this, &FPhysicsAssetEditor::OnToggleShowSelected),
+		FCanExecuteAction::CreateSP(this, &FPhysicsAssetEditor::IsNotSimulation));
+
 	ToolkitCommands->MapAction(
 		Commands.DeselectAll,
 		FExecuteAction::CreateSP(this, &FPhysicsAssetEditor::OnDeselectAll),
@@ -1240,10 +1246,11 @@ void FPhysicsAssetEditor::BuildMenuWidgetSelection(FMenuBuilder& InMenuBuilder)
 
 		InMenuBuilder.BeginSection( "Selection", LOCTEXT("Selection", "Selection" ) );
 		InMenuBuilder.AddMenuEntry( Commands.SelectAllBodies );
-		InMenuBuilder.AddMenuEntry(Commands.SelectSimulatedBodies);
-		InMenuBuilder.AddMenuEntry(Commands.SelectKinematicBodies);
+		InMenuBuilder.AddMenuEntry( Commands.SelectSimulatedBodies );
+		InMenuBuilder.AddMenuEntry( Commands.SelectKinematicBodies );
 		InMenuBuilder.AddMenuEntry( Commands.SelectAllConstraints );
 		InMenuBuilder.AddMenuEntry( Commands.ToggleSelectionType );
+		InMenuBuilder.AddMenuEntry( Commands.ToggleShowSelected );
 		InMenuBuilder.EndSection();
 	}
 	InMenuBuilder.PopCommandList();
@@ -2688,6 +2695,11 @@ void FPhysicsAssetEditor::OnToggleSelectionType()
 	SharedData->ToggleSelectionType();
 }
 
+void FPhysicsAssetEditor::OnToggleShowSelected()
+{
+	SharedData->ToggleShowSelected();
+}
+
 void FPhysicsAssetEditor::OnDeselectAll()
 {
 	SharedData->ClearSelectedBody();
diff --git a/Engine/Source/Editor/PhysicsAssetEditor/Private/PhysicsAssetEditor.h b/Engine/Source/Editor/PhysicsAssetEditor/Private/PhysicsAssetEditor.h
--- a/Engine/Source/Editor/PhysicsAssetEditor/Private/PhysicsAssetEditor.h
+++ b/Engine/Source/Editor/PhysicsAssetEditor/Private/PhysicsAssetEditor.h
@@ -274,6 +274,7 @@ private:
 	void OnSelectBodies(EPhysicsType PhysicsType = EPhysicsType::PhysType_Simulated);
 	void OnSelectAllConstraints();
 	void OnToggleSelectionType();
+	void OnToggleShowSelected();
 	void OnDeselectAll();
 
 	FText GetRepeatLastSimulationToolTip() const;
@@ -351,4 +352,9 @@ private:
 
 	/** Records PhysicsAssetEditor related data - simulating or mode change */
 	void OnAddPhatRecord(const FString& Action, bool bRecordSimulate, bool bRecordMode);
+
+public:
+
+	TArray<int32> HiddenBodies;
+	TArray<int32> HiddenConstraints;
 };
diff --git a/Engine/Source/Editor/PhysicsAssetEditor/Private/PhysicsAssetEditorActions.cpp b/Engine/Source/Editor/PhysicsAssetEditor/Private/PhysicsAssetEditorActions.cpp
--- a/Engine/Source/Editor/PhysicsAssetEditor/Private/PhysicsAssetEditorActions.cpp
+++ b/Engine/Source/Editor/PhysicsAssetEditor/Private/PhysicsAssetEditorActions.cpp
@@ -75,6 +75,7 @@ void FPhysicsAssetEditorCommands::RegisterCommands()
 	UI_COMMAND(SelectKinematicBodies, "Select Kinematic Bodies", "Select Kinematic Bodies", EUserInterfaceActionType::Button, FInputChord(EModifierKey::Control, EKeys::K));
 	UI_COMMAND(SelectAllConstraints, "Select All Constraints", "Select All Constraints", EUserInterfaceActionType::Button, FInputChord(EModifierKey::Control | EModifierKey::Shift, EKeys::A));
 	UI_COMMAND(ToggleSelectionType, "Toggle Selection Type", "Select bodies from constraints, or constraints from bodies, depending on selection.", EUserInterfaceActionType::Button, FInputChord(EModifierKey::Control, EKeys::T));
+	UI_COMMAND(ToggleShowSelected, "Toggle Show Selected", "Show/hide selected bodies/constraints.", EUserInterfaceActionType::Button, FInputChord(EModifierKey::Control, EKeys::H));
 	UI_COMMAND(DeselectAll, "Deselect All Bodies", "Deselect All", EUserInterfaceActionType::Button, FInputChord(EKeys::Escape));
 	UI_COMMAND(Mirror, "Mirror", "Finds the body on the other side and duplicates constraint and body", EUserInterfaceActionType::Button, FInputChord(EKeys::M));
 
diff --git a/Engine/Source/Editor/PhysicsAssetEditor/Private/PhysicsAssetEditorActions.h b/Engine/Source/Editor/PhysicsAssetEditor/Private/PhysicsAssetEditorActions.h
--- a/Engine/Source/Editor/PhysicsAssetEditor/Private/PhysicsAssetEditorActions.h
+++ b/Engine/Source/Editor/PhysicsAssetEditor/Private/PhysicsAssetEditorActions.h
@@ -87,6 +87,7 @@ public:
 	TSharedPtr<FUICommandInfo> SelectKinematicBodies;
 	TSharedPtr<FUICommandInfo> SelectAllConstraints;
 	TSharedPtr<FUICommandInfo> ToggleSelectionType;
+	TSharedPtr<FUICommandInfo> ToggleShowSelected;
 	TSharedPtr<FUICommandInfo> DeselectAll;
 	TSharedPtr<FUICommandInfo> Mirror;
 	TSharedPtr<FUICommandInfo> NewPhysicalAnimationProfile;
diff --git a/Engine/Source/Editor/PhysicsAssetEditor/Private/PhysicsAssetEditorSharedData.cpp b/Engine/Source/Editor/PhysicsAssetEditor/Private/PhysicsAssetEditorSharedData.cpp
--- a/Engine/Source/Editor/PhysicsAssetEditor/Private/PhysicsAssetEditorSharedData.cpp
+++ b/Engine/Source/Editor/PhysicsAssetEditor/Private/PhysicsAssetEditorSharedData.cpp
@@ -554,6 +554,32 @@ void FPhysicsAssetEditorSharedData::ToggleSelectionType()
 	}
 }
 
+void FPhysicsAssetEditorSharedData::ToggleShowSelected()
+{
+	for (const FSelection& Selection : SelectedConstraints)
+	{
+		if (HiddenConstraints.Contains(Selection.Index))
+		{
+			HiddenConstraints.Remove(Selection.Index);
+		}
+		else
+		{
+			HiddenConstraints.AddUnique(Selection.Index);
+		}
+	}
+	for (const FSelection& Selection : SelectedBodies)
+	{
+		if (HiddenBodies.Contains(Selection.Index))
+		{
+			HiddenBodies.Remove(Selection.Index);
+		}
+		else
+		{
+			HiddenBodies.AddUnique(Selection.Index);
+		}
+	}
+}
+
 void FPhysicsAssetEditorSharedData::UpdateNoCollisionBodies()
 {
 	NoCollisionBodies.Empty();
diff --git a/Engine/Source/Editor/PhysicsAssetEditor/Private/PhysicsAssetEditorSharedData.h b/Engine/Source/Editor/PhysicsAssetEditor/Private/PhysicsAssetEditorSharedData.h
--- a/Engine/Source/Editor/PhysicsAssetEditor/Private/PhysicsAssetEditorSharedData.h
+++ b/Engine/Source/Editor/PhysicsAssetEditor/Private/PhysicsAssetEditorSharedData.h
@@ -144,6 +144,7 @@ public:
 	void SetSelectedBody(const FSelection& Body, bool bSelected);
 	bool IsBodySelected(const FSelection& Body) const;
 	void ToggleSelectionType();
+	void ToggleShowSelected();
 	void SetSelectedBodyAnyPrim(int32 BodyIndex, bool bSelected);
 	void DeleteCurrentPrim();
 	void DeleteBody(int32 DelBodyIndex, bool bRefreshComponent=true);
@@ -259,6 +260,9 @@ public:
 	TArray<FBoneVertInfo> AnyWeightBoneInfos;
 
 	TArray<FSelection> SelectedBodies;
+
+	TArray<int32> HiddenBodies;
+	TArray<int32> HiddenConstraints;
 	FSelection * GetSelectedBody()
 	{
 		int32 Count = SelectedBodies.Num();
diff --git a/Engine/Source/Editor/PhysicsAssetEditor/Private/PhysicsAssetEditorSkeletalMeshComponent.cpp b/Engine/Source/Editor/PhysicsAssetEditor/Private/PhysicsAssetEditorSkeletalMeshComponent.cpp
--- a/Engine/Source/Editor/PhysicsAssetEditor/Private/PhysicsAssetEditorSkeletalMeshComponent.cpp
+++ b/Engine/Source/Editor/PhysicsAssetEditor/Private/PhysicsAssetEditorSkeletalMeshComponent.cpp
@@ -116,6 +116,10 @@ void UPhysicsAssetEditorSkeletalMeshComponent::RenderAssetTools(const FSceneView
 		{
 			continue;
 		}
+		if (SharedData->HiddenBodies.Contains(i))
+		{
+			continue;
+		}
 		int32 BoneIndex = GetBoneIndex(PhysicsAsset->SkeletalBodySetups[i]->BoneName);
 
 		// If we found a bone for it, draw the collision.
@@ -240,7 +244,8 @@ void UPhysicsAssetEditorSkeletalMeshComponent::RenderAssetTools(const FSceneView
 	{
 		for (int32 i = 0; i <PhysicsAsset->ConstraintSetup.Num(); ++i)
 		{
-			if(!SharedData->EditorOptions->bRenderOnlySelectedConstraints || (SharedData->EditorOptions->bRenderOnlySelectedConstraints && SharedData->IsConstraintSelected(i)))
+			if((!SharedData->EditorOptions->bRenderOnlySelectedConstraints || (SharedData->EditorOptions->bRenderOnlySelectedConstraints && SharedData->IsConstraintSelected(i))) &&
+				!SharedData->HiddenConstraints.Contains(i))
 			{
 				int32 BoneIndex1 = GetBoneIndex(PhysicsAsset->ConstraintSetup[i]->DefaultInstance.ConstraintBone1);
 				int32 BoneIndex2 = GetBoneIndex(PhysicsAsset->ConstraintSetup[i]->DefaultInstance.ConstraintBone2);
