Set better defaults for physics asset editor skeleton tree.

diff --git a/Engine/Source/Editor/PhysicsAssetEditor/Private/PhysicsAssetEditor.cpp b/Engine/Source/Editor/PhysicsAssetEditor/Private/PhysicsAssetEditor.cpp
--- a/Engine/Source/Editor/PhysicsAssetEditor/Private/PhysicsAssetEditor.cpp
+++ b/Engine/Source/Editor/PhysicsAssetEditor/Private/PhysicsAssetEditor.cpp
@@ -79,6 +79,7 @@
 #include "Widgets/Input/SSpinBox.h"
 #include "Widgets/Input/SNumericEntryBox.h"
 #include "AnimationEditorPreviewActor.h"
+#include "Preferences/PersonaOptions.h"
 
 const FName PhysicsAssetEditorModes::PhysicsAssetEditorMode("PhysicsAssetEditorMode");
 
@@ -148,6 +149,10 @@ void FPhysicsAssetEditor::InitPhysicsAssetEditor(const EToolkitMode::Type Mode,
 	SkeletonTreeArgs.ContextName = GetToolkitFName();
 
 	ISkeletonEditorModule& SkeletonEditorModule = FModuleManager::GetModuleChecked<ISkeletonEditorModule>("SkeletonEditor");
+
+	GetMutableDefault<UPersonaOptions>()->bFlattenSkeletonHierarchyWhenFiltering = false;
+	GetMutableDefault<UPersonaOptions>()->bHideParentsWhenFiltering = true;
+
 	SkeletonTree = SkeletonEditorModule.CreateSkeletonTree(PersonaToolkit->GetSkeleton(), SkeletonTreeArgs);
 
 	bSelecting = false;
@@ -1278,6 +1283,7 @@ TSharedRef<ISkeletonTree> FPhysicsAssetEditor::BuildMenuWidgetNewConstraintForBo
 	SkeletonTreeArgs.bAllowSkeletonOperations = false;
 	SkeletonTreeArgs.bShowBlendProfiles = false;
 	SkeletonTreeArgs.bShowFilterMenu = false;
+	SkeletonTreeArgs.bHideBonesByDefault = false;
 	SkeletonTreeArgs.Builder = Builder;
 	SkeletonTreeArgs.PreviewScene = GetPersonaToolkit()->GetPreviewScene();
 	SkeletonTreeArgs.OnSelectionChanged = FOnSkeletonTreeSelectionChanged::CreateLambda([this, InSourceBodyIndex, InOnActionMenuClosed](const TArrayView<TSharedPtr<ISkeletonTreeItem>>& InSelectedItems, ESelectInfo::Type SelectInfo)
diff --git a/Engine/Source/Editor/PhysicsAssetEditor/Private/PhysicsAssetEditorSkeletonTreeBuilder.cpp b/Engine/Source/Editor/PhysicsAssetEditor/Private/PhysicsAssetEditorSkeletonTreeBuilder.cpp
--- a/Engine/Source/Editor/PhysicsAssetEditor/Private/PhysicsAssetEditorSkeletonTreeBuilder.cpp
+++ b/Engine/Source/Editor/PhysicsAssetEditor/Private/PhysicsAssetEditorSkeletonTreeBuilder.cpp
@@ -16,8 +16,8 @@ FPhysicsAssetEditorSkeletonTreeBuilder::FPhysicsAssetEditorSkeletonTreeBuilder(U
 	, bShowBodies(true)
 	, bShowKinematicBodies(true)
 	, bShowSimulatedBodies(true)
-	, bShowConstraints(true)
-	, bShowPrimitives(true)
+	, bShowConstraints(false)
+	, bShowPrimitives(false)
 	, PhysicsAsset(InPhysicsAsset)
 {
 }
diff --git a/Engine/Source/Editor/SkeletonEditor/Private/SSkeletonTree.cpp b/Engine/Source/Editor/SkeletonEditor/Private/SSkeletonTree.cpp
--- a/Engine/Source/Editor/SkeletonEditor/Private/SSkeletonTree.cpp
+++ b/Engine/Source/Editor/SkeletonEditor/Private/SSkeletonTree.cpp
@@ -123,7 +123,14 @@ public:
 
 void SSkeletonTree::Construct(const FArguments& InArgs, const TSharedRef<FEditableSkeleton>& InEditableSkeleton, const FSkeletonTreeArgs& InSkeletonTreeArgs)
 {
-	BoneFilter = EBoneFilter::All;
+	if (InSkeletonTreeArgs.bHideBonesByDefault)
+	{
+		BoneFilter = EBoneFilter::None;
+	}
+	else
+	{
+		BoneFilter = EBoneFilter::All;
+	}
 	SocketFilter = ESocketFilter::Active;
 	bShowingAdvancedOptions = false;
 	bSelecting = false;
diff --git a/Engine/Source/Editor/SkeletonEditor/Public/ISkeletonTree.h b/Engine/Source/Editor/SkeletonEditor/Public/ISkeletonTree.h
--- a/Engine/Source/Editor/SkeletonEditor/Public/ISkeletonTree.h
+++ b/Engine/Source/Editor/SkeletonEditor/Public/ISkeletonTree.h
@@ -50,6 +50,7 @@ struct FSkeletonTreeArgs
 		, bShowFilterMenu(true)
 		, bAllowMeshOperations(true)
 		, bAllowSkeletonOperations(true)
+		, bHideBonesByDefault(true)
 	{}
 
 	PRAGMA_DISABLE_DEPRECATION_WARNINGS
@@ -92,6 +93,9 @@ struct FSkeletonTreeArgs
 
 	/** Whether to allow operations that modify the skeleton */
 	bool bAllowSkeletonOperations;
+
+	/** Whether to hide bones by default */
+	bool bHideBonesByDefault;
 };
 
 /** Interface used to deal with skeleton editing UI */
