Fix skeleton tree hide bones default to be false.

diff --git a/Engine/Source/Editor/PhysicsAssetEditor/Private/PhysicsAssetEditor.cpp b/Engine/Source/Editor/PhysicsAssetEditor/Private/PhysicsAssetEditor.cpp
--- a/Engine/Source/Editor/PhysicsAssetEditor/Private/PhysicsAssetEditor.cpp
+++ b/Engine/Source/Editor/PhysicsAssetEditor/Private/PhysicsAssetEditor.cpp
@@ -141,6 +141,7 @@ void FPhysicsAssetEditor::InitPhysicsAssetEditor(const EToolkitMode::Type Mode,
 	SkeletonTreeArgs.bShowBlendProfiles = false;
 	SkeletonTreeArgs.bAllowMeshOperations = false;
 	SkeletonTreeArgs.bAllowSkeletonOperations = false;
+	SkeletonTreeArgs.bHideBonesByDefault = true;
 	SkeletonTreeArgs.OnGetFilterText = FOnGetFilterText::CreateSP(this, &FPhysicsAssetEditor::HandleGetFilterLabel);
 	SkeletonTreeArgs.Extenders = MakeShared<FExtender>();
 	SkeletonTreeArgs.Extenders->AddMenuExtension("FilterOptions", EExtensionHook::After, GetToolkitCommands(), FMenuExtensionDelegate::CreateSP(this, &FPhysicsAssetEditor::HandleExtendFilterMenu));
@@ -1283,7 +1284,7 @@ TSharedRef<ISkeletonTree> FPhysicsAssetEditor::BuildMenuWidgetNewConstraintForBo
 	SkeletonTreeArgs.bAllowSkeletonOperations = false;
 	SkeletonTreeArgs.bShowBlendProfiles = false;
 	SkeletonTreeArgs.bShowFilterMenu = false;
-	SkeletonTreeArgs.bHideBonesByDefault = false;
+	SkeletonTreeArgs.bHideBonesByDefault = true;
 	SkeletonTreeArgs.Builder = Builder;
 	SkeletonTreeArgs.PreviewScene = GetPersonaToolkit()->GetPreviewScene();
 	SkeletonTreeArgs.OnSelectionChanged = FOnSkeletonTreeSelectionChanged::CreateLambda([this, InSourceBodyIndex, InOnActionMenuClosed](const TArrayView<TSharedPtr<ISkeletonTreeItem>>& InSelectedItems, ESelectInfo::Type SelectInfo)
diff --git a/Engine/Source/Editor/SkeletonEditor/Public/ISkeletonTree.h b/Engine/Source/Editor/SkeletonEditor/Public/ISkeletonTree.h
--- a/Engine/Source/Editor/SkeletonEditor/Public/ISkeletonTree.h
+++ b/Engine/Source/Editor/SkeletonEditor/Public/ISkeletonTree.h
@@ -50,7 +50,7 @@ struct FSkeletonTreeArgs
 		, bShowFilterMenu(true)
 		, bAllowMeshOperations(true)
 		, bAllowSkeletonOperations(true)
-		, bHideBonesByDefault(true)
+		, bHideBonesByDefault(false)
 	{}
 
 	PRAGMA_DISABLE_DEPRECATION_WARNINGS
