Chaos - Chaos PhAt Simulate uses Rigid Body Node

diff --git a/Engine/Source/Editor/PhysicsAssetEditor/Private/PhysicsAssetEditor.cpp b/Engine/Source/Editor/PhysicsAssetEditor/Private/PhysicsAssetEditor.cpp
--- a/Engine/Source/Editor/PhysicsAssetEditor/Private/PhysicsAssetEditor.cpp
+++ b/Engine/Source/Editor/PhysicsAssetEditor/Private/PhysicsAssetEditor.cpp
@@ -2933,11 +2933,7 @@ void FPhysicsAssetEditor::HandlePreviewSceneCreated(const TSharedRef<IPersonaPre
 
 	SharedData->EditorSkelComp->Stop();
 
-	SharedData->EditorSkelComp->SetAnimationMode(EAnimationMode::AnimationCustomMode);
 	SharedData->EditorSkelComp->PreviewInstance = NewObject<UPhysicsAssetEditorAnimInstance>(SharedData->EditorSkelComp, TEXT("PhatAnimScriptInstance"));
-	SharedData->EditorSkelComp->AnimScriptInstance = SharedData->EditorSkelComp->PreviewInstance;
-	SharedData->EditorSkelComp->AnimScriptInstance->InitializeAnimation();
-	SharedData->EditorSkelComp->InitAnim(true);
 
 	SharedData->PhysicalAnimationComponent = NewObject<UPhysicalAnimationComponent>(Actor);
 	SharedData->PhysicalAnimationComponent->SetSkeletalMeshComponent(SharedData->EditorSkelComp);
diff --git a/Engine/Source/Editor/PhysicsAssetEditor/Private/PhysicsAssetEditorAnimInstanceProxy.cpp b/Engine/Source/Editor/PhysicsAssetEditor/Private/PhysicsAssetEditorAnimInstanceProxy.cpp
--- a/Engine/Source/Editor/PhysicsAssetEditor/Private/PhysicsAssetEditorAnimInstanceProxy.cpp
+++ b/Engine/Source/Editor/PhysicsAssetEditor/Private/PhysicsAssetEditorAnimInstanceProxy.cpp
@@ -2,7 +2,6 @@
 
 #include "PhysicsAssetEditorAnimInstanceProxy.h"
 #include "PhysicsAssetEditorAnimInstance.h"
-//#include "AnimNode_RigidBody.h"
 
 void FPhysicsAssetEditorAnimInstanceProxy::Initialize(UAnimInstance* InAnimInstance)
 {
@@ -12,17 +11,46 @@ void FPhysicsAssetEditorAnimInstanceProxy::Initialize(UAnimInstance* InAnimInsta
 
 void FPhysicsAssetEditorAnimInstanceProxy::ConstructNodes()
 {
-	//LocalToComponentSpace.LocalPose.SetLinkNode(&SequencePlayerNode);
-	//RagdollNode.ComponentPose.SetLinkNode(&LocalToComponentSpace);
-	//ComponentToLocalSpace.ComponentPose.SetLinkNode(&RagdollNode);
-	//ComponentToLocalSpace.ComponentPose.SetLinkNode(&LocalToComponentSpace);
+	ComponentToLocalSpace.ComponentPose.SetLinkNode(&RagdollNode);
 	
-	//RagdollNode.SimulationSpace = ESimulationSpace::WorldSpace;
+	RagdollNode.SimulationSpace = ESimulationSpace::WorldSpace;
+	RagdollNode.ActualAlpha = 1.0f;
+}
+
+FAnimNode_Base* FPhysicsAssetEditorAnimInstanceProxy::GetCustomRootNode()
+{
+	return &ComponentToLocalSpace;
 }
 
 void FPhysicsAssetEditorAnimInstanceProxy::GetCustomNodes(TArray<FAnimNode_Base*>& OutNodes)
 {
-	//OutNodes.Add(&RagdollNode);
+	OutNodes.Add(&RagdollNode);
+	OutNodes.Add(&ComponentToLocalSpace);
+}
+
+void FPhysicsAssetEditorAnimInstanceProxy::UpdateAnimationNode(const FAnimationUpdateContext& InContext)
+{
+	if (CurrentAsset != nullptr)
+	{
+		FAnimPreviewInstanceProxy::UpdateAnimationNode(InContext);
+	}
+	else
+	{
+		ComponentToLocalSpace.Update_AnyThread(InContext);
+	}
+}
+
+bool FPhysicsAssetEditorAnimInstanceProxy::Evaluate_WithRoot(FPoseContext& Output, FAnimNode_Base* InRootNode)
+{
+	if (CurrentAsset != nullptr)
+	{
+		return FAnimPreviewInstanceProxy::Evaluate_WithRoot(Output, InRootNode);
+	}
+	else
+	{
+		InRootNode->Evaluate_AnyThread(Output);
+		return true;
+	}
 }
 
 FPhysicsAssetEditorAnimInstanceProxy::~FPhysicsAssetEditorAnimInstanceProxy()
diff --git a/Engine/Source/Editor/PhysicsAssetEditor/Private/PhysicsAssetEditorAnimInstanceProxy.h b/Engine/Source/Editor/PhysicsAssetEditor/Private/PhysicsAssetEditorAnimInstanceProxy.h
--- a/Engine/Source/Editor/PhysicsAssetEditor/Private/PhysicsAssetEditorAnimInstanceProxy.h
+++ b/Engine/Source/Editor/PhysicsAssetEditor/Private/PhysicsAssetEditorAnimInstanceProxy.h
@@ -3,9 +3,10 @@
 #pragma once
 
 #include "AnimPreviewInstance.h"
-//#include "AnimNode_RigidBody.h"
-#include "Animation/AnimNodeSpaceConversions.h"
+#include "Animation/AnimNode_LinkedInputPose.h"
 #include "Animation/AnimNode_SequencePlayer.h"
+#include "Animation/AnimNodeSpaceConversions.h"
+#include "BoneControllers/AnimNode_RigidBody.h"
 #include "PhysicsAssetEditorAnimInstanceProxy.generated.h"
 
 class UAnimSequence;
@@ -29,15 +30,15 @@ public:
 	virtual ~FPhysicsAssetEditorAnimInstanceProxy();
 
 	virtual void Initialize(UAnimInstance* InAnimInstance) override;
+	virtual FAnimNode_Base* GetCustomRootNode() override;
 	virtual void GetCustomNodes(TArray<FAnimNode_Base*>& OutNodes) override;
+
+	virtual void UpdateAnimationNode(const FAnimationUpdateContext& InContext) override;
+	virtual bool Evaluate_WithRoot(FPoseContext& Output, FAnimNode_Base* InRootNode) override;
+
 private:
 	void ConstructNodes();
 
-	//FAnimNode_RigidBody RagdollNode;
-// 	FAnimNode_ConvertComponentToLocalSpace ComponentToLocalSpace;
-// 	FAnimNode_SequencePlayer SequencePlayerNode;
-// 	FAnimNode_ConvertLocalToComponentSpace LocalToComponentSpace;
-	
-	void InitAnimTrack(UAnimSequenceBase* InAnimSequence, int32 SequenceId);
-	void EnsureAnimTrack(UAnimSequenceBase* InAnimSequence, int32 SequenceId);
+	FAnimNode_RigidBody RagdollNode;
+ 	FAnimNode_ConvertComponentToLocalSpace ComponentToLocalSpace;
 };
\ No newline at end of file
diff --git a/Engine/Source/Editor/PhysicsAssetEditor/Private/PhysicsAssetEditorSharedData.cpp b/Engine/Source/Editor/PhysicsAssetEditor/Private/PhysicsAssetEditorSharedData.cpp
--- a/Engine/Source/Editor/PhysicsAssetEditor/Private/PhysicsAssetEditorSharedData.cpp
+++ b/Engine/Source/Editor/PhysicsAssetEditor/Private/PhysicsAssetEditorSharedData.cpp
@@ -36,6 +36,10 @@
 
 #define LOCTEXT_NAMESPACE "PhysicsAssetEditorShared"
 
+// Whether to use RigidBody AnimNode for simulation preview when using Chaos since we don't have constraints in the main scene yet.
+// NOTE: The SkeletalMeshComponent simulation overrides the AnimNode simulation if enabled, so PHAT_USE_RBAN_SIMULATION switches it off.
+#define PHAT_USE_RBAN_SIMULATION WITH_CHAOS
+
 FScopedBulkSelection::FScopedBulkSelection(TSharedPtr<FPhysicsAssetEditorSharedData> InSharedData)
 	: SharedData(InSharedData)
 {
@@ -1678,7 +1682,7 @@ void FPhysicsAssetEditorSharedData::EnableSimulation(bool bEnableSimulation)
 {
 	if (bEnableSimulation)
 	{
-
+#if !PHAT_USE_RBAN_SIMULATION
 		// We should not already have an instance (destroyed when stopping sim).
 		EditorSkelComp->SetCollisionEnabled(ECollisionEnabled::QueryAndPhysics);
 		EditorSkelComp->SetSimulatePhysics(true);
@@ -1688,6 +1692,11 @@ void FPhysicsAssetEditorSharedData::EnableSimulation(bool bEnableSimulation)
 
 		// Make it start simulating
 		EditorSkelComp->WakeAllRigidBodies();
+#else
+		// Enable the PreviewInstance (containing the AnimNode_RigidBody)
+		EditorSkelComp->SetAnimationMode(EAnimationMode::AnimationCustomMode);
+		EditorSkelComp->InitAnim(true);
+#endif
 
 		if(EditorOptions->bResetClothWhenSimulating)
 		{
@@ -1696,9 +1705,9 @@ void FPhysicsAssetEditorSharedData::EnableSimulation(bool bEnableSimulation)
 	}
 	else
 	{
-		EditorSkelComp->AnimScriptInstance = nullptr;
-
-		if(EditorSkelComp->GetAnimationMode() != EAnimationMode::AnimationSingleNode)
+		// Disable the PreviewInstance
+		//EditorSkelComp->AnimScriptInstance = nullptr;
+		//if(EditorSkelComp->GetAnimationMode() != EAnimationMode::AnimationSingleNode)
 		{
 			EditorSkelComp->SetAnimationMode(EAnimationMode::AnimationSingleNode);
 		}
diff --git a/Engine/Source/Runtime/AnimGraphRuntime/Private/BoneControllers/AnimNode_RigidBody.cpp b/Engine/Source/Runtime/AnimGraphRuntime/Private/BoneControllers/AnimNode_RigidBody.cpp
--- a/Engine/Source/Runtime/AnimGraphRuntime/Private/BoneControllers/AnimNode_RigidBody.cpp
+++ b/Engine/Source/Runtime/AnimGraphRuntime/Private/BoneControllers/AnimNode_RigidBody.cpp
@@ -36,6 +36,7 @@ FAutoConsoleVariableRef CVarRigidBodyNodeEnableTimeBasedReset(TEXT("p.RigidBodyN
 FAnimNode_RigidBody::FAnimNode_RigidBody():
 	QueryParams(NAME_None, FCollisionQueryParams::GetUnknownStatId())
 {
+	AccumulatedDeltaTime = 0.0f;
 	ResetSimulatedTeleportType = ETeleportType::None;
 	PhysicsSimulation = nullptr;
 	OverridePhysicsAsset = nullptr;
@@ -50,6 +51,7 @@ FAnimNode_RigidBody::FAnimNode_RigidBody():
 	OverrideWorldGravity = FVector::ZeroVector;
 	TotalMass = 0.f;
 	CachedBounds.W = 0;
+	PhysScene = nullptr;
 	UnsafeWorld = nullptr;
 	bSimulationStarted = false;
 	bCheckForBodyTransformInit = false;
diff --git a/Engine/Source/Runtime/AnimGraphRuntime/Private/BoneControllers/AnimNode_RigidBody_Chaos.cpp b/Engine/Source/Runtime/AnimGraphRuntime/Private/BoneControllers/AnimNode_RigidBody_Chaos.cpp
--- a/Engine/Source/Runtime/AnimGraphRuntime/Private/BoneControllers/AnimNode_RigidBody_Chaos.cpp
+++ b/Engine/Source/Runtime/AnimGraphRuntime/Private/BoneControllers/AnimNode_RigidBody_Chaos.cpp
@@ -28,6 +28,7 @@ TAutoConsoleVariable<int32> CVarChaosRigidBodyLODThreshold(TEXT("p.ChaosRigidBod
 FAnimNode_RigidBody_Chaos::FAnimNode_RigidBody_Chaos()
 	: QueryParams(NAME_None, FCollisionQueryParams::GetUnknownStatId())
 {
+	AccumulatedDeltaTime = 0.0f;
 	ResetSimulatedTeleportType = ETeleportType::None;
 	PhysicsSimulation = nullptr;
 	OverridePhysicsAsset = nullptr;
@@ -41,6 +42,7 @@ FAnimNode_RigidBody_Chaos::FAnimNode_RigidBody_Chaos()
 	OverrideWorldGravity = FVector::ZeroVector;
 	TotalMass = 0.f;
 	CachedBounds.W = 0;
+	PhysScene = nullptr;
 	UnsafeWorld = nullptr;
 	bSimulationStarted = false;
 	bCheckForBodyTransformInit = false;
