Chaos - hook solver iterations up to PhAt preview

diff --git a/Engine/Source/Editor/PhysicsAssetEditor/Private/PhysicsAssetEditorAnimInstanceProxy.cpp b/Engine/Source/Editor/PhysicsAssetEditor/Private/PhysicsAssetEditorAnimInstanceProxy.cpp
--- a/Engine/Source/Editor/PhysicsAssetEditor/Private/PhysicsAssetEditorAnimInstanceProxy.cpp
+++ b/Engine/Source/Editor/PhysicsAssetEditor/Private/PhysicsAssetEditorAnimInstanceProxy.cpp
@@ -31,6 +31,14 @@ void FPhysicsAssetEditorAnimInstanceProxy::Initialize(UAnimInstance* InAnimInsta
 {
 	FAnimPreviewInstanceProxy::Initialize(InAnimInstance);
 	ConstructNodes();
+
+#if WITH_CHAOS && !PHYSICS_INTERFACE_PHYSX
+	UPhysicsAsset* PhysicsAsset = InAnimInstance->GetSkelMeshComponent()->GetPhysicsAsset();
+	if (PhysicsAsset != nullptr)
+	{
+		SolverIterations = PhysicsAsset->SolverIterations;
+	}
+#endif
 }
 
 void FPhysicsAssetEditorAnimInstanceProxy::ConstructNodes()
@@ -66,6 +74,21 @@ void FPhysicsAssetEditorAnimInstanceProxy::UpdateAnimationNode(const FAnimationU
 
 bool FPhysicsAssetEditorAnimInstanceProxy::Evaluate_WithRoot(FPoseContext& Output, FAnimNode_Base* InRootNode)
 {
+#if WITH_CHAOS && !PHYSICS_INTERFACE_PHYSX
+	ImmediatePhysics::FSimulation* Simulation = RagdollNode.GetSimulation();
+	if (Simulation != nullptr)
+	{
+		Simulation->SetSolverIterations(
+			SolverIterations.SolverIterations,
+			SolverIterations.JointIterations,
+			SolverIterations.CollisionIterations,
+			SolverIterations.SolverPushOutIterations,
+			SolverIterations.JointPushOutIterations,
+			SolverIterations.CollisionPushOutIterations
+		);
+	}
+#endif
+
 	if (CurrentAsset != nullptr)
 	{
 		return FAnimPreviewInstanceProxy::Evaluate_WithRoot(Output, InRootNode);
diff --git a/Engine/Source/Editor/PhysicsAssetEditor/Private/PhysicsAssetEditorAnimInstanceProxy.h b/Engine/Source/Editor/PhysicsAssetEditor/Private/PhysicsAssetEditorAnimInstanceProxy.h
--- a/Engine/Source/Editor/PhysicsAssetEditor/Private/PhysicsAssetEditorAnimInstanceProxy.h
+++ b/Engine/Source/Editor/PhysicsAssetEditor/Private/PhysicsAssetEditorAnimInstanceProxy.h
@@ -45,6 +45,7 @@ private:
 
 	FAnimNode_RigidBody RagdollNode;
  	FAnimNode_ConvertComponentToLocalSpace ComponentToLocalSpace;
+	FSolverIterations SolverIterations;
 
 	ImmediatePhysics::FActorHandle* TargetActor;
 	ImmediatePhysics::FActorHandle* HandleActor;
