PhAT - fix crash inserting bodies when joint constraint's have been manually modified to reference other bodies (and therefore the constraint hierarchy does not match the bone hierarchy)

diff --git a/Engine/Source/Editor/PhysicsAssetEditor/Private/PhysicsAssetEditorSharedData.cpp b/Engine/Source/Editor/PhysicsAssetEditor/Private/PhysicsAssetEditorSharedData.cpp
--- a/Engine/Source/Editor/PhysicsAssetEditor/Private/PhysicsAssetEditorSharedData.cpp
+++ b/Engine/Source/Editor/PhysicsAssetEditor/Private/PhysicsAssetEditorSharedData.cpp
@@ -1080,6 +1080,13 @@ void FPhysicsAssetEditorSharedData::MakeNewBody(int32 NewBoneIndex, bool bAutoSe
 				
 				// If the child body is not constrained already, create a new constraint between
 				// the child body and the new body
+				// @todo: This isn't quite right. It is possible that the child constraint's parent body is not our parent body. 
+				// This can happen in a couple ways:
+				// - the user altered the child constraint to attach to a different parent bond
+				// - a new bone was added. E.g., add bone at root of hierarchy. Import mesh with new bone. Add body to root bone.
+				// So, if this happens we need to decide if we should leave the old constraint there and add a new one, or commandeer the
+				// constraint. If the former, we should probably change a constraint to a "User" constraint when they change its bones.
+				// We are currently doing the latter...
 				if (ConstraintIndex == INDEX_NONE)
 				{
 					ConstraintIndex = FPhysicsAssetUtils::CreateNewConstraint(PhysicsAsset, ChildBody->BoneName);
@@ -1099,16 +1106,19 @@ void FPhysicsAssetEditorSharedData::MakeNewBody(int32 NewBoneIndex, bool bAutoSe
 					{
 						continue;
 					}
-					
+
 					// If the constraint isn't between two child bones, then it is between a physical bone higher in the bone
 					// hierarchy than the new bone, so it needs to be fixed up by setting the constraint to point to the new bone
 					// instead. Additionally, collision needs to be re-enabled between the child bone and the identified "grandparent"
 					// bone.
 					const int32 ExistingConstraintBodyIndex = PhysicsAsset->FindBodyIndex(ExistingConstraintSetup->DefaultInstance.ConstraintBone2);
 					check(ExistingConstraintBodyIndex != INDEX_NONE);
-					check(ExistingConstraintBodyIndex == ParentBodyIndex);
 
-					SetCollisionBetween(ChildBodyIndex, ExistingConstraintBodyIndex, true);
+					// See above comments about the child constraint's parent not necessarily being our parent...
+					if (ExistingConstraintBodyIndex == ParentBodyIndex)
+					{
+						SetCollisionBetween(ChildBodyIndex, ExistingConstraintBodyIndex, true);
+					}
 				}
 
 				UPhysicsConstraintTemplate* ChildConstraintSetup = PhysicsAsset->ConstraintSetup[ ConstraintIndex ];
