Add r-click menu option to change primitive collision filtering in PHAT

diff --git a/Engine/Source/Editor/PhysicsAssetEditor/Private/PhysicsAssetEditor.cpp b/Engine/Source/Editor/PhysicsAssetEditor/Private/PhysicsAssetEditor.cpp
--- a/Engine/Source/Editor/PhysicsAssetEditor/Private/PhysicsAssetEditor.cpp
+++ b/Engine/Source/Editor/PhysicsAssetEditor/Private/PhysicsAssetEditor.cpp
@@ -670,6 +670,30 @@ void FPhysicsAssetEditor::BindCommands()
 		FExecuteAction::CreateSP(this, &FPhysicsAssetEditor::OnSetCollisionAll, true),
 		FCanExecuteAction::CreateSP(this, &FPhysicsAssetEditor::CanSetCollisionAll, true));
 
+	ToolkitCommands->MapAction(
+		Commands.PrimitiveQueryAndPhysics,
+		FExecuteAction::CreateSP(this, &FPhysicsAssetEditor::OnSetPrimitiveCollision, ECollisionEnabled::QueryAndPhysics),
+		FCanExecuteAction::CreateSP(this, &FPhysicsAssetEditor::CanSetPrimitiveCollision, ECollisionEnabled::QueryAndPhysics),
+		FIsActionChecked::CreateSP(this, &FPhysicsAssetEditor::IsPrimitiveCollisionChecked, ECollisionEnabled::QueryAndPhysics));
+
+	ToolkitCommands->MapAction(
+		Commands.PrimitiveQueryOnly,
+		FExecuteAction::CreateSP(this, &FPhysicsAssetEditor::OnSetPrimitiveCollision, ECollisionEnabled::QueryOnly),
+		FCanExecuteAction::CreateSP(this, &FPhysicsAssetEditor::CanSetPrimitiveCollision, ECollisionEnabled::QueryOnly),
+		FIsActionChecked::CreateSP(this, &FPhysicsAssetEditor::IsPrimitiveCollisionChecked, ECollisionEnabled::QueryOnly));
+
+	ToolkitCommands->MapAction(
+		Commands.PrimitivePhysicsOnly,
+		FExecuteAction::CreateSP(this, &FPhysicsAssetEditor::OnSetPrimitiveCollision, ECollisionEnabled::PhysicsOnly),
+		FCanExecuteAction::CreateSP(this, &FPhysicsAssetEditor::CanSetPrimitiveCollision, ECollisionEnabled::PhysicsOnly),
+		FIsActionChecked::CreateSP(this, &FPhysicsAssetEditor::IsPrimitiveCollisionChecked, ECollisionEnabled::PhysicsOnly));
+
+	ToolkitCommands->MapAction(
+		Commands.PrimitiveNoCollision,
+		FExecuteAction::CreateSP(this, &FPhysicsAssetEditor::OnSetPrimitiveCollision, ECollisionEnabled::NoCollision),
+		FCanExecuteAction::CreateSP(this, &FPhysicsAssetEditor::CanSetPrimitiveCollision, ECollisionEnabled::NoCollision),
+		FIsActionChecked::CreateSP(this, &FPhysicsAssetEditor::IsPrimitiveCollisionChecked, ECollisionEnabled::NoCollision));
+
 	ToolkitCommands->MapAction(
 		Commands.WeldToBody,
 		FExecuteAction::CreateSP(this, &FPhysicsAssetEditor::OnWeldToBody),
@@ -1126,6 +1150,15 @@ void FPhysicsAssetEditor::BuildMenuWidgetBody(FMenuBuilder& InMenuBuilder)
 				InSubMenuBuilder.AddMenuEntry( PhysicsAssetEditorCommands.DisableCollision );
 				InSubMenuBuilder.AddMenuEntry( PhysicsAssetEditorCommands.DisableCollisionAll );
 				InSubMenuBuilder.EndSection();
+
+#if WITH_CHAOS
+				InSubMenuBuilder.BeginSection("CollisionFilteringHeader", LOCTEXT("CollisionFilteringHeader", "CollisionFiltering"));
+				InSubMenuBuilder.AddMenuEntry(PhysicsAssetEditorCommands.PrimitiveQueryAndPhysics);
+				InSubMenuBuilder.AddMenuEntry(PhysicsAssetEditorCommands.PrimitiveQueryOnly);
+				InSubMenuBuilder.AddMenuEntry(PhysicsAssetEditorCommands.PrimitivePhysicsOnly);
+				InSubMenuBuilder.AddMenuEntry(PhysicsAssetEditorCommands.PrimitiveNoCollision);
+				InSubMenuBuilder.EndSection();
+#endif
 			}
 		};
 
@@ -2155,6 +2188,23 @@ bool FPhysicsAssetEditor::CanSetCollisionAll(bool bEnable) const
 	return SharedData->CanSetCollisionBetweenSelectedAndAll(bEnable);
 }
 
+void FPhysicsAssetEditor::OnSetPrimitiveCollision(ECollisionEnabled::Type CollisionEnabled)
+{
+	FScopedTransaction Transaction(LOCTEXT("SetPrimitiveCollision", "Set Primitive Collision"));
+
+	SharedData->SetPrimitiveCollision(CollisionEnabled);
+}
+
+bool FPhysicsAssetEditor::CanSetPrimitiveCollision(ECollisionEnabled::Type CollisionEnabled) const
+{
+	return SharedData->CanSetPrimitiveCollision(CollisionEnabled);
+}
+
+bool FPhysicsAssetEditor::IsPrimitiveCollisionChecked(ECollisionEnabled::Type CollisionEnabled) const
+{
+	return SharedData->GetIsPrimitiveCollisionEnabled(CollisionEnabled);
+}
+
 void FPhysicsAssetEditor::OnWeldToBody()
 {
 	SharedData->WeldSelectedBodies();
diff --git a/Engine/Source/Editor/PhysicsAssetEditor/Private/PhysicsAssetEditor.h b/Engine/Source/Editor/PhysicsAssetEditor/Private/PhysicsAssetEditor.h
--- a/Engine/Source/Editor/PhysicsAssetEditor/Private/PhysicsAssetEditor.h
+++ b/Engine/Source/Editor/PhysicsAssetEditor/Private/PhysicsAssetEditor.h
@@ -233,6 +233,9 @@ private:
 	bool CanSetCollision(bool bEnable) const;
 	void OnSetCollisionAll(bool bEnable);
 	bool CanSetCollisionAll(bool bEnable) const;
+	void OnSetPrimitiveCollision(ECollisionEnabled::Type CollisionEnabled);
+	bool CanSetPrimitiveCollision(ECollisionEnabled::Type CollisionEnabled) const;
+	bool IsPrimitiveCollisionChecked(ECollisionEnabled::Type CollisionEnabled) const;
 	void OnWeldToBody();
 	bool CanWeldToBody();
 	void OnAddSphere();
diff --git a/Engine/Source/Editor/PhysicsAssetEditor/Private/PhysicsAssetEditorActions.cpp b/Engine/Source/Editor/PhysicsAssetEditor/Private/PhysicsAssetEditorActions.cpp
--- a/Engine/Source/Editor/PhysicsAssetEditor/Private/PhysicsAssetEditorActions.cpp
+++ b/Engine/Source/Editor/PhysicsAssetEditor/Private/PhysicsAssetEditorActions.cpp
@@ -53,6 +53,10 @@ void FPhysicsAssetEditorCommands::RegisterCommands()
 	UI_COMMAND(AddBox, "Add Box", "Add Box To Selected Bone", EUserInterfaceActionType::Button, FInputChord());
 	UI_COMMAND(AddTaperedCapsule, "Add Tapered Capsule (Clothing Only)", "Add Tapered Capsule To Selected Bone. This is only used by clothing, so will have no effect on rigid body collision, overlaps or bounds.", EUserInterfaceActionType::Button, FInputChord());
 	UI_COMMAND(DeletePrimitive, "Delete", "Delete Selected Primitive(s)", EUserInterfaceActionType::Button, FInputChord());
+	UI_COMMAND(PrimitiveQueryAndPhysics, "All Collision", "Enable all collision on selected primitive(s)", EUserInterfaceActionType::RadioButton, FInputChord());
+	UI_COMMAND(PrimitiveQueryOnly, "Query Collision Only", "Enable query collision only on selected primitive(s)", EUserInterfaceActionType::RadioButton, FInputChord());
+	UI_COMMAND(PrimitivePhysicsOnly, "Physics Collision Only", "Enable physics collision only on selected primitive(s)", EUserInterfaceActionType::RadioButton, FInputChord());
+	UI_COMMAND(PrimitiveNoCollision, "No Collision", "Disable all collision on Selected primitive(s)", EUserInterfaceActionType::RadioButton, FInputChord());
 	UI_COMMAND(DuplicatePrimitive, "Duplicate", "Duplicate Selected Primitive(s)", EUserInterfaceActionType::Button, FInputChord());
 	UI_COMMAND(ResetConstraint, "Reset", "Reset Constraint", EUserInterfaceActionType::Button, FInputChord());
 	UI_COMMAND(SnapConstraint, "Snap", "Snap Constraint Translation To Bone", EUserInterfaceActionType::Button, FInputChord());
diff --git a/Engine/Source/Editor/PhysicsAssetEditor/Private/PhysicsAssetEditorActions.h b/Engine/Source/Editor/PhysicsAssetEditor/Private/PhysicsAssetEditorActions.h
--- a/Engine/Source/Editor/PhysicsAssetEditor/Private/PhysicsAssetEditorActions.h
+++ b/Engine/Source/Editor/PhysicsAssetEditor/Private/PhysicsAssetEditorActions.h
@@ -59,6 +59,10 @@ public:
 	TSharedPtr<FUICommandInfo> DisableCollisionAll;
 	TSharedPtr<FUICommandInfo> EnableCollision;
 	TSharedPtr<FUICommandInfo> EnableCollisionAll;
+	TSharedPtr<FUICommandInfo> PrimitiveQueryAndPhysics;
+	TSharedPtr<FUICommandInfo> PrimitiveQueryOnly;
+	TSharedPtr<FUICommandInfo> PrimitivePhysicsOnly;
+	TSharedPtr<FUICommandInfo> PrimitiveNoCollision;
 	TSharedPtr<FUICommandInfo> WeldToBody;
 	TSharedPtr<FUICommandInfo> AddSphere;
 	TSharedPtr<FUICommandInfo> AddSphyl;
diff --git a/Engine/Source/Editor/PhysicsAssetEditor/Private/PhysicsAssetEditorSharedData.cpp b/Engine/Source/Editor/PhysicsAssetEditor/Private/PhysicsAssetEditorSharedData.cpp
--- a/Engine/Source/Editor/PhysicsAssetEditor/Private/PhysicsAssetEditorSharedData.cpp
+++ b/Engine/Source/Editor/PhysicsAssetEditor/Private/PhysicsAssetEditorSharedData.cpp
@@ -807,6 +807,46 @@ void FPhysicsAssetEditorSharedData::SetCollisionBetween(int32 Body1Index, int32
 	PreviewChangedEvent.Broadcast();
 }
 
+void FPhysicsAssetEditorSharedData::SetPrimitiveCollision(ECollisionEnabled::Type CollisionEnabled)
+{
+	if (bRunningSimulation)
+	{
+		return;
+	}
+
+	PhysicsAsset->Modify();
+
+	for (FSelection SelectedBody : SelectedBodies)
+	{
+		PhysicsAsset->SetPrimitiveCollision(SelectedBody.Index, SelectedBody.PrimitiveIndex, CollisionEnabled);
+	}
+
+	PreviewChangedEvent.Broadcast();
+}
+
+bool FPhysicsAssetEditorSharedData::CanSetPrimitiveCollision(ECollisionEnabled::Type CollisionEnabled) const
+{
+	if (bRunningSimulation || SelectedBodies.Num() == 0)
+	{
+		return false;
+	}
+
+	return true;
+}
+
+bool FPhysicsAssetEditorSharedData::GetIsPrimitiveCollisionEnabled(ECollisionEnabled::Type CollisionEnabled) const
+{
+	for (const FSelection SelectedBody : SelectedBodies)
+	{
+		if (PhysicsAsset->GetPrimitiveCollision(SelectedBody.Index, SelectedBody.PrimitiveIndex) == CollisionEnabled)
+		{
+			return true;
+		}
+	}
+
+	return false;
+}
+
 void FPhysicsAssetEditorSharedData::CopyBody()
 {
 	check(SelectedBodies.Num() == 1);
diff --git a/Engine/Source/Editor/PhysicsAssetEditor/Private/PhysicsAssetEditorSharedData.h b/Engine/Source/Editor/PhysicsAssetEditor/Private/PhysicsAssetEditorSharedData.h
--- a/Engine/Source/Editor/PhysicsAssetEditor/Private/PhysicsAssetEditorSharedData.h
+++ b/Engine/Source/Editor/PhysicsAssetEditor/Private/PhysicsAssetEditorSharedData.h
@@ -185,6 +185,11 @@ public:
 	void SetCollisionBetweenSelectedAndAll(bool bEnableCollision);
 	bool CanSetCollisionBetweenSelectedAndAll(bool bEnableCollision) const;
 
+	/** Helpers to set primitive-level collision filtering on selected bodies */
+	void SetPrimitiveCollision(ECollisionEnabled::Type CollisionEnabled);
+	bool CanSetPrimitiveCollision(ECollisionEnabled::Type CollisionEnabled) const;
+	bool GetIsPrimitiveCollisionEnabled(ECollisionEnabled::Type CollisionEnabled) const;
+
 	/** Prevents GC from collecting our objects */
 	void AddReferencedObjects(FReferenceCollector& Collector);
 
diff --git a/Engine/Source/Runtime/Engine/Classes/PhysicsEngine/AggregateGeom.h b/Engine/Source/Runtime/Engine/Classes/PhysicsEngine/AggregateGeom.h
--- a/Engine/Source/Runtime/Engine/Classes/PhysicsEngine/AggregateGeom.h
+++ b/Engine/Source/Runtime/Engine/Classes/PhysicsEngine/AggregateGeom.h
@@ -61,6 +61,28 @@ struct ENGINE_API FKAggregateGeom
 
 	int32 GetElementCount(EAggCollisionShape::Type Type) const;
 
+	FKShapeElem* GetElement(const int32 Index)
+	{
+		ensure(Index < GetElementCount());
+		if (Index < SphereElems.Num()) { return &SphereElems[Index]; }
+		if (Index < BoxElems.Num()) { return &BoxElems[Index]; }
+		if (Index < SphylElems.Num()) { return &SphylElems[Index]; }
+		if (Index < ConvexElems.Num()) { return &ConvexElems[Index]; }
+		if (Index < TaperedCapsuleElems.Num()) { return &TaperedCapsuleElems[Index]; }
+		return nullptr;
+	}
+
+	const FKShapeElem* GetElement(const int32 Index) const
+	{
+		ensure(Index < GetElementCount());
+		if (Index < SphereElems.Num()) { return &SphereElems[Index]; }
+		if (Index < BoxElems.Num()) { return &BoxElems[Index]; }
+		if (Index < SphylElems.Num()) { return &SphylElems[Index]; }
+		if (Index < ConvexElems.Num()) { return &ConvexElems[Index]; }
+		if (Index < TaperedCapsuleElems.Num()) { return &TaperedCapsuleElems[Index]; }
+		return nullptr;
+	}
+
 	void EmptyElements()
 	{
 		BoxElems.Empty();
diff --git a/Engine/Source/Runtime/Engine/Classes/PhysicsEngine/PhysicsAsset.h b/Engine/Source/Runtime/Engine/Classes/PhysicsEngine/PhysicsAsset.h
--- a/Engine/Source/Runtime/Engine/Classes/PhysicsEngine/PhysicsAsset.h
+++ b/Engine/Source/Runtime/Engine/Classes/PhysicsEngine/PhysicsAsset.h
@@ -228,6 +228,12 @@ public:
 	// Check whether the two bodies specified are enabled for collision
 	ENGINE_API bool IsCollisionEnabled(int32 BodyIndexA, int32 BodyIndexB) const;
 
+	// Get the per-primitive collision filtering mode for a body
+	ENGINE_API void SetPrimitiveCollision(int32 BodyIndex, int32 PrimitiveIndex, ECollisionEnabled::Type CollisionEnabled);
+
+	// Get the per-primitive collision filtering mode for a body
+	ENGINE_API ECollisionEnabled::Type GetPrimitiveCollision(int32 BodyIndex, int32 PrimitiveIndex) const;
+
 	/** Update the BoundsBodies array and cache the indices of bodies marked with bConsiderForBounds to BoundsBodies array. */
 	ENGINE_API void UpdateBoundsBodiesArray();
 
diff --git a/Engine/Source/Runtime/Engine/Private/PhysicsEngine/PhysicsAsset.cpp b/Engine/Source/Runtime/Engine/Private/PhysicsEngine/PhysicsAsset.cpp
--- a/Engine/Source/Runtime/Engine/Private/PhysicsEngine/PhysicsAsset.cpp
+++ b/Engine/Source/Runtime/Engine/Private/PhysicsEngine/PhysicsAsset.cpp
@@ -253,6 +253,24 @@ bool UPhysicsAsset::IsCollisionEnabled(int32 BodyIndexA, int32 BodyIndexB) const
 	return true;
 }
 
+void UPhysicsAsset::SetPrimitiveCollision(int32 BodyIndex, int32 PrimitiveIndex, ECollisionEnabled::Type CollisionEnabled)
+{
+#if WITH_CHAOS
+	check(SkeletalBodySetups.IsValidIndex(BodyIndex));
+	FKAggregateGeom* AggGeom = &SkeletalBodySetups[BodyIndex]->AggGeom;
+	ensure(PrimitiveIndex < AggGeom->GetElementCount());
+	AggGeom->GetElement(PrimitiveIndex)->SetCollisionEnabled(CollisionEnabled);
+#endif
+}
+
+ECollisionEnabled::Type UPhysicsAsset::GetPrimitiveCollision(int32 BodyIndex, int32 PrimitiveIndex) const
+{
+	check(SkeletalBodySetups.IsValidIndex(BodyIndex));
+	FKAggregateGeom* AggGeom = &SkeletalBodySetups[BodyIndex]->AggGeom;
+	ensure(PrimitiveIndex < AggGeom->GetElementCount());
+	return AggGeom->GetElement(PrimitiveIndex)->GetCollisionEnabled();
+}
+
 FBox UPhysicsAsset::CalcAABB(const USkinnedMeshComponent* MeshComp, const FTransform& LocalToWorld) const
 {
 	FBox Box(ForceInit);
