PhAT - Add new Show/Hide options

diff --git a/Engine/Source/Editor/PhysicsAssetEditor/Private/PhysicsAssetEditor.cpp b/Engine/Source/Editor/PhysicsAssetEditor/Private/PhysicsAssetEditor.cpp
--- a/Engine/Source/Editor/PhysicsAssetEditor/Private/PhysicsAssetEditor.cpp
+++ b/Engine/Source/Editor/PhysicsAssetEditor/Private/PhysicsAssetEditor.cpp
@@ -584,6 +584,11 @@ void FPhysicsAssetEditor::ExtendMenu()
 			MenuBarBuilder.AddMenuEntry(Commands.SelectAllConstraints);
 			MenuBarBuilder.AddMenuEntry(Commands.ToggleSelectionType);
 			MenuBarBuilder.AddMenuEntry(Commands.ToggleShowSelected);
+			MenuBarBuilder.AddMenuEntry(Commands.ShowAll);
+			MenuBarBuilder.AddMenuEntry(Commands.HideAll);
+			MenuBarBuilder.AddMenuEntry(Commands.ShowSelected);
+			MenuBarBuilder.AddMenuEntry(Commands.HideSelected);
+			MenuBarBuilder.AddMenuEntry(Commands.ShowOnlySelected);
 			MenuBarBuilder.AddMenuEntry(Commands.DeselectAll);
 			MenuBarBuilder.EndSection();
 		}
@@ -870,6 +875,32 @@ void FPhysicsAssetEditor::BindCommands()
 		FExecuteAction::CreateSP(this, &FPhysicsAssetEditor::OnToggleShowSelected),
 		FCanExecuteAction::CreateSP(this, &FPhysicsAssetEditor::IsNotSimulation));
 
+	ToolkitCommands->MapAction(
+		Commands.ShowAll,
+		FExecuteAction::CreateSP(this, &FPhysicsAssetEditor::OnShowAll),
+		FCanExecuteAction::CreateSP(this, &FPhysicsAssetEditor::IsNotSimulation));
+
+	ToolkitCommands->MapAction(
+		Commands.HideAll,
+		FExecuteAction::CreateSP(this, &FPhysicsAssetEditor::OnHideAll),
+		FCanExecuteAction::CreateSP(this, &FPhysicsAssetEditor::IsNotSimulation));
+
+	ToolkitCommands->MapAction(
+		Commands.ShowSelected,
+		FExecuteAction::CreateSP(this, &FPhysicsAssetEditor::OnShowSelected),
+		FCanExecuteAction::CreateSP(this, &FPhysicsAssetEditor::IsNotSimulation));
+
+	ToolkitCommands->MapAction(
+		Commands.HideSelected,
+		FExecuteAction::CreateSP(this, &FPhysicsAssetEditor::OnHideSelected),
+		FCanExecuteAction::CreateSP(this, &FPhysicsAssetEditor::IsNotSimulation));
+
+
+	ToolkitCommands->MapAction(
+		Commands.ShowOnlySelected,
+		FExecuteAction::CreateSP(this, &FPhysicsAssetEditor::OnShowOnlySelected),
+		FCanExecuteAction::CreateSP(this, &FPhysicsAssetEditor::IsNotSimulation));
+
 	ToolkitCommands->MapAction(
 		Commands.DeselectAll,
 		FExecuteAction::CreateSP(this, &FPhysicsAssetEditor::OnDeselectAll),
@@ -1298,6 +1329,11 @@ void FPhysicsAssetEditor::BuildMenuWidgetSelection(FMenuBuilder& InMenuBuilder)
 		InMenuBuilder.AddMenuEntry( Commands.SelectAllConstraints );
 		InMenuBuilder.AddMenuEntry( Commands.ToggleSelectionType );
 		InMenuBuilder.AddMenuEntry( Commands.ToggleShowSelected );
+		InMenuBuilder.AddMenuEntry( Commands.ShowAll );
+		InMenuBuilder.AddMenuEntry( Commands.HideAll );
+		InMenuBuilder.AddMenuEntry( Commands.ShowSelected );
+		InMenuBuilder.AddMenuEntry( Commands.HideSelected );
+		InMenuBuilder.AddMenuEntry( Commands.ShowOnlySelected );
 		InMenuBuilder.EndSection();
 	}
 	InMenuBuilder.PopCommandList();
@@ -2765,6 +2801,33 @@ void FPhysicsAssetEditor::OnToggleShowSelected()
 	SharedData->ToggleShowSelected();
 }
 
+void FPhysicsAssetEditor::OnShowAll()
+{
+	SharedData->ShowAll();
+}
+
+void FPhysicsAssetEditor::OnHideAll()
+{
+	SharedData->HideAll();
+}
+
+void FPhysicsAssetEditor::OnShowSelected()
+{
+	SharedData->ShowSelected();
+}
+
+void FPhysicsAssetEditor::OnHideSelected()
+{
+	SharedData->HideSelected();
+}
+
+
+void FPhysicsAssetEditor::OnShowOnlySelected()
+{
+	SharedData->HideAll();
+	SharedData->ShowSelected();
+}
+
 void FPhysicsAssetEditor::OnDeselectAll()
 {
 	SharedData->ClearSelectedBody();
diff --git a/Engine/Source/Editor/PhysicsAssetEditor/Private/PhysicsAssetEditor.h b/Engine/Source/Editor/PhysicsAssetEditor/Private/PhysicsAssetEditor.h
--- a/Engine/Source/Editor/PhysicsAssetEditor/Private/PhysicsAssetEditor.h
+++ b/Engine/Source/Editor/PhysicsAssetEditor/Private/PhysicsAssetEditor.h
@@ -278,6 +278,11 @@ private:
 	void OnSelectAllConstraints();
 	void OnToggleSelectionType();
 	void OnToggleShowSelected();
+	void OnShowAll();
+	void OnHideAll();
+	void OnShowSelected();
+	void OnHideSelected();
+	void OnShowOnlySelected();
 	void OnDeselectAll();
 
 	FText GetRepeatLastSimulationToolTip() const;
diff --git a/Engine/Source/Editor/PhysicsAssetEditor/Private/PhysicsAssetEditorActions.cpp b/Engine/Source/Editor/PhysicsAssetEditor/Private/PhysicsAssetEditorActions.cpp
--- a/Engine/Source/Editor/PhysicsAssetEditor/Private/PhysicsAssetEditorActions.cpp
+++ b/Engine/Source/Editor/PhysicsAssetEditor/Private/PhysicsAssetEditorActions.cpp
@@ -79,7 +79,12 @@ void FPhysicsAssetEditorCommands::RegisterCommands()
 	UI_COMMAND(SelectKinematicBodies, "Select Kinematic Bodies", "Select Kinematic Bodies", EUserInterfaceActionType::Button, FInputChord(EModifierKey::Control, EKeys::K));
 	UI_COMMAND(SelectAllConstraints, "Select All Constraints", "Select All Constraints", EUserInterfaceActionType::Button, FInputChord(EModifierKey::Control | EModifierKey::Shift, EKeys::A));
 	UI_COMMAND(ToggleSelectionType, "Toggle Selection Type", "Select bodies from constraints, or constraints from bodies, depending on selection.", EUserInterfaceActionType::Button, FInputChord(EModifierKey::Control, EKeys::T));
-	UI_COMMAND(ToggleShowSelected, "Toggle Show Selected", "Show/hide selected bodies/constraints.", EUserInterfaceActionType::Button, FInputChord(EModifierKey::Control, EKeys::H));
+	UI_COMMAND(ToggleShowSelected, "Toggle Show Selected", "Show/hide selected bodies/constraints.", EUserInterfaceActionType::Button, FInputChord(EModifierKey::Control, EKeys::R));
+	UI_COMMAND(ShowAll, "Show All", "Show all bodies/constraints.", EUserInterfaceActionType::Button, FInputChord(EModifierKey::Shift, EKeys::G));
+	UI_COMMAND(HideAll, "Hide All", "Hide all bodies/constraints.", EUserInterfaceActionType::Button, FInputChord(EModifierKey::Control, EKeys::G));
+	UI_COMMAND(ShowSelected, "Show Selected", "Show selected bodies/constraints.", EUserInterfaceActionType::Button, FInputChord(EModifierKey::Shift, EKeys::H));
+	UI_COMMAND(HideSelected, "Hide Selected", "Hide selected bodies/constraints.", EUserInterfaceActionType::Button, FInputChord(EModifierKey::Control, EKeys::H));
+	UI_COMMAND(ShowOnlySelected, "Show Only Selected", "Show only the selected bodies/constraints, hiding all others.", EUserInterfaceActionType::Button, FInputChord(EKeys::I));
 	UI_COMMAND(DeselectAll, "Deselect All Bodies", "Deselect All", EUserInterfaceActionType::Button, FInputChord(EKeys::Escape));
 	UI_COMMAND(Mirror, "Mirror", "Finds the body on the other side and duplicates constraint and body", EUserInterfaceActionType::Button, FInputChord(EKeys::M));
 
diff --git a/Engine/Source/Editor/PhysicsAssetEditor/Private/PhysicsAssetEditorActions.h b/Engine/Source/Editor/PhysicsAssetEditor/Private/PhysicsAssetEditorActions.h
--- a/Engine/Source/Editor/PhysicsAssetEditor/Private/PhysicsAssetEditorActions.h
+++ b/Engine/Source/Editor/PhysicsAssetEditor/Private/PhysicsAssetEditorActions.h
@@ -92,6 +92,11 @@ public:
 	TSharedPtr<FUICommandInfo> SelectAllConstraints;
 	TSharedPtr<FUICommandInfo> ToggleSelectionType;
 	TSharedPtr<FUICommandInfo> ToggleShowSelected;
+	TSharedPtr<FUICommandInfo> ShowAll;
+	TSharedPtr<FUICommandInfo> HideAll;
+	TSharedPtr<FUICommandInfo> ShowSelected;
+	TSharedPtr<FUICommandInfo> HideSelected;
+	TSharedPtr<FUICommandInfo> ShowOnlySelected;
 	TSharedPtr<FUICommandInfo> DeselectAll;
 	TSharedPtr<FUICommandInfo> Mirror;
 	TSharedPtr<FUICommandInfo> NewPhysicalAnimationProfile;
diff --git a/Engine/Source/Editor/PhysicsAssetEditor/Private/PhysicsAssetEditorSharedData.cpp b/Engine/Source/Editor/PhysicsAssetEditor/Private/PhysicsAssetEditorSharedData.cpp
--- a/Engine/Source/Editor/PhysicsAssetEditor/Private/PhysicsAssetEditorSharedData.cpp
+++ b/Engine/Source/Editor/PhysicsAssetEditor/Private/PhysicsAssetEditorSharedData.cpp
@@ -589,6 +589,67 @@ void FPhysicsAssetEditorSharedData::ToggleShowSelected()
 	}
 }
 
+void FPhysicsAssetEditorSharedData::ShowAll()
+{
+	HiddenConstraints.Empty();
+	HiddenBodies.Empty();
+}
+
+
+void FPhysicsAssetEditorSharedData::HideAll()
+{
+	if (PhysicsAsset != nullptr)
+	{
+		HiddenBodies.Empty();
+		for (int32 i = 0; i < PhysicsAsset->SkeletalBodySetups.Num(); ++i)
+		{
+			HiddenBodies.Add(i);
+		}
+
+		HiddenConstraints.Empty();
+		for (int32 i = 0; i < PhysicsAsset->ConstraintSetup.Num(); ++i)
+		{
+			HiddenConstraints.Add(i);
+		}
+	}
+}
+
+void FPhysicsAssetEditorSharedData::ShowSelected()
+{
+	for (const FSelection& Selection : SelectedConstraints)
+	{
+		if (HiddenConstraints.Contains(Selection.Index))
+		{
+			HiddenConstraints.Remove(Selection.Index);
+		}
+	}
+	for (const FSelection& Selection : SelectedBodies)
+	{
+		if (HiddenBodies.Contains(Selection.Index))
+		{
+			HiddenBodies.Remove(Selection.Index);
+		}
+	}
+}
+
+void FPhysicsAssetEditorSharedData::HideSelected()
+{
+	for (const FSelection& Selection : SelectedConstraints)
+	{
+		if (!HiddenConstraints.Contains(Selection.Index))
+		{
+			HiddenConstraints.Add(Selection.Index);
+		}
+	}
+	for (const FSelection& Selection : SelectedBodies)
+	{
+		if (!HiddenBodies.Contains(Selection.Index))
+		{
+			HiddenBodies.Add(Selection.Index);
+		}
+	}
+}
+
 void FPhysicsAssetEditorSharedData::UpdateNoCollisionBodies()
 {
 	NoCollisionBodies.Empty();
diff --git a/Engine/Source/Editor/PhysicsAssetEditor/Private/PhysicsAssetEditorSharedData.h b/Engine/Source/Editor/PhysicsAssetEditor/Private/PhysicsAssetEditorSharedData.h
--- a/Engine/Source/Editor/PhysicsAssetEditor/Private/PhysicsAssetEditorSharedData.h
+++ b/Engine/Source/Editor/PhysicsAssetEditor/Private/PhysicsAssetEditorSharedData.h
@@ -145,6 +145,10 @@ public:
 	bool IsBodySelected(const FSelection& Body) const;
 	void ToggleSelectionType();
 	void ToggleShowSelected();
+	void ShowAll();
+	void HideAll();
+	void ShowSelected();
+	void HideSelected();
 	void SetSelectedBodyAnyPrim(int32 BodyIndex, bool bSelected);
 	void DeleteCurrentPrim();
 	void DeleteBody(int32 DelBodyIndex, bool bRefreshComponent=true);
