PhAT - change "Show Only Selected" shortcut key. It was being blocked nu "Show Environment" if the viewport window was active.

diff --git a/Engine/Source/Editor/PhysicsAssetEditor/Private/PhysicsAssetEditor.cpp b/Engine/Source/Editor/PhysicsAssetEditor/Private/PhysicsAssetEditor.cpp
--- a/Engine/Source/Editor/PhysicsAssetEditor/Private/PhysicsAssetEditor.cpp
+++ b/Engine/Source/Editor/PhysicsAssetEditor/Private/PhysicsAssetEditor.cpp
@@ -88,6 +88,8 @@ const FName PhysicsAssetEditorAppIdentifier = FName(TEXT("PhysicsAssetEditorApp"
 DEFINE_LOG_CATEGORY(LogPhysicsAssetEditor);
 #define LOCTEXT_NAMESPACE "PhysicsAssetEditor"
 
+//PRAGMA_DISABLE_OPTIMIZATION
+
 namespace PhysicsAssetEditor
 {
 	const float	DefaultPrimSize = 15.0f;
diff --git a/Engine/Source/Editor/PhysicsAssetEditor/Private/PhysicsAssetEditorActions.cpp b/Engine/Source/Editor/PhysicsAssetEditor/Private/PhysicsAssetEditorActions.cpp
--- a/Engine/Source/Editor/PhysicsAssetEditor/Private/PhysicsAssetEditorActions.cpp
+++ b/Engine/Source/Editor/PhysicsAssetEditor/Private/PhysicsAssetEditorActions.cpp
@@ -85,7 +85,7 @@ void FPhysicsAssetEditorCommands::RegisterCommands()
 	UI_COMMAND(HideAll, "Hide All", "Hide all bodies/constraints.", EUserInterfaceActionType::Button, FInputChord(EModifierKey::Control, EKeys::G));
 	UI_COMMAND(ShowSelected, "Show Selected", "Show selected bodies/constraints.", EUserInterfaceActionType::Button, FInputChord(EModifierKey::Shift, EKeys::H));
 	UI_COMMAND(HideSelected, "Hide Selected", "Hide selected bodies/constraints.", EUserInterfaceActionType::Button, FInputChord(EModifierKey::Control, EKeys::H));
-	UI_COMMAND(ShowOnlySelected, "Show Only Selected", "Show only the selected bodies/constraints, hiding all others.", EUserInterfaceActionType::Button, FInputChord(EKeys::I));
+	UI_COMMAND(ShowOnlySelected, "Show Only Selected", "Show only the selected bodies/constraints, hiding all others.", EUserInterfaceActionType::Button, FInputChord(EKeys::G));
 	UI_COMMAND(DeselectAll, "Deselect All Bodies", "Deselect All", EUserInterfaceActionType::Button, FInputChord(EKeys::Escape));
 	UI_COMMAND(Mirror, "Mirror", "Finds the body on the other side and duplicates constraint and body", EUserInterfaceActionType::Button, FInputChord(EKeys::M));
 
diff --git a/Engine/Source/Editor/PhysicsAssetEditor/Private/PhysicsAssetEditorSharedData.cpp b/Engine/Source/Editor/PhysicsAssetEditor/Private/PhysicsAssetEditorSharedData.cpp
--- a/Engine/Source/Editor/PhysicsAssetEditor/Private/PhysicsAssetEditorSharedData.cpp
+++ b/Engine/Source/Editor/PhysicsAssetEditor/Private/PhysicsAssetEditorSharedData.cpp
@@ -36,6 +36,8 @@
 
 #define LOCTEXT_NAMESPACE "PhysicsAssetEditorShared"
 
+//PRAGMA_DISABLE_OPTIMIZATION
+
 // Whether to use RigidBody AnimNode for simulation preview when using Chaos since we don't have constraints in the main scene yet.
 // NOTE: The SkeletalMeshComponent simulation overrides the AnimNode simulation if enabled, so PHAT_USE_RBAN_SIMULATION switches it off.
 #define PHAT_USE_RBAN_SIMULATION WITH_CHAOS
