PhAT - key shortcut changes

diff --git a/Engine/Source/Editor/PhysicsAssetEditor/Private/PhysicsAssetEditorActions.cpp b/Engine/Source/Editor/PhysicsAssetEditor/Private/PhysicsAssetEditorActions.cpp
--- a/Engine/Source/Editor/PhysicsAssetEditor/Private/PhysicsAssetEditorActions.cpp
+++ b/Engine/Source/Editor/PhysicsAssetEditor/Private/PhysicsAssetEditorActions.cpp
@@ -80,12 +80,12 @@ void FPhysicsAssetEditorCommands::RegisterCommands()
 	UI_COMMAND(SelectKinematicBodies, "Select Kinematic Bodies", "Select Kinematic Bodies", EUserInterfaceActionType::Button, FInputChord(EModifierKey::Control, EKeys::K));
 	UI_COMMAND(SelectAllConstraints, "Select All Constraints", "Select All Constraints", EUserInterfaceActionType::Button, FInputChord(EModifierKey::Control | EModifierKey::Shift, EKeys::A));
 	UI_COMMAND(ToggleSelectionType, "Toggle Selection Type", "Select bodies from constraints, or constraints from bodies, depending on selection.", EUserInterfaceActionType::Button, FInputChord(EModifierKey::Control, EKeys::T));
-	UI_COMMAND(ToggleShowSelected, "Toggle Show Selected", "Show/hide selected bodies/constraints.", EUserInterfaceActionType::Button, FInputChord(EModifierKey::Control, EKeys::R));
-	UI_COMMAND(ShowAll, "Show All", "Show all bodies/constraints.", EUserInterfaceActionType::Button, FInputChord(EModifierKey::Shift, EKeys::G));
-	UI_COMMAND(HideAll, "Hide All", "Hide all bodies/constraints.", EUserInterfaceActionType::Button, FInputChord(EModifierKey::Control, EKeys::G));
+	UI_COMMAND(ToggleShowSelected, "Toggle Show Selected", "Show/hide selected bodies/constraints.", EUserInterfaceActionType::Button, FInputChord(EKeys::H));
 	UI_COMMAND(ShowSelected, "Show Selected", "Show selected bodies/constraints.", EUserInterfaceActionType::Button, FInputChord(EModifierKey::Shift, EKeys::H));
 	UI_COMMAND(HideSelected, "Hide Selected", "Hide selected bodies/constraints.", EUserInterfaceActionType::Button, FInputChord(EModifierKey::Control, EKeys::H));
 	UI_COMMAND(ShowOnlySelected, "Show Only Selected", "Show only the selected bodies/constraints, hiding all others.", EUserInterfaceActionType::Button, FInputChord(EKeys::G));
+	UI_COMMAND(ShowAll, "Show All", "Show all bodies/constraints.", EUserInterfaceActionType::Button, FInputChord(EModifierKey::Shift, EKeys::G));
+	UI_COMMAND(HideAll, "Hide All", "Hide all bodies/constraints.", EUserInterfaceActionType::Button, FInputChord(EModifierKey::Control, EKeys::G));
 	UI_COMMAND(DeselectAll, "Deselect All Bodies", "Deselect All", EUserInterfaceActionType::Button, FInputChord(EKeys::Escape));
 	UI_COMMAND(Mirror, "Mirror", "Finds the body on the other side and duplicates constraint and body", EUserInterfaceActionType::Button, FInputChord(EKeys::M));
 
