PhAT - Isolate is now a toggle

diff --git a/Engine/Source/Editor/PhysicsAssetEditor/Private/PhysicsAssetEditor.cpp b/Engine/Source/Editor/PhysicsAssetEditor/Private/PhysicsAssetEditor.cpp
--- a/Engine/Source/Editor/PhysicsAssetEditor/Private/PhysicsAssetEditor.cpp
+++ b/Engine/Source/Editor/PhysicsAssetEditor/Private/PhysicsAssetEditor.cpp
@@ -588,7 +588,7 @@ void FPhysicsAssetEditor::ExtendMenu()
 			MenuBarBuilder.AddMenuEntry(Commands.ToggleShowSelected);
 			MenuBarBuilder.AddMenuEntry(Commands.ShowSelected);
 			MenuBarBuilder.AddMenuEntry(Commands.HideSelected);
-			MenuBarBuilder.AddMenuEntry(Commands.ShowOnlySelected);
+			MenuBarBuilder.AddMenuEntry(Commands.ToggleShowOnlySelected);
 			MenuBarBuilder.AddMenuEntry(Commands.ShowAll);
 			MenuBarBuilder.AddMenuEntry(Commands.HideAll);
 			MenuBarBuilder.AddMenuEntry(Commands.DeselectAll);
@@ -894,8 +894,8 @@ void FPhysicsAssetEditor::BindCommands()
 		FCanExecuteAction::CreateSP(this, &FPhysicsAssetEditor::IsNotSimulation));
 
 	ToolkitCommands->MapAction(
-		Commands.ShowOnlySelected,
-		FExecuteAction::CreateSP(this, &FPhysicsAssetEditor::OnShowOnlySelected),
+		Commands.ToggleShowOnlySelected,
+		FExecuteAction::CreateSP(this, &FPhysicsAssetEditor::OnToggleShowOnlySelected),
 		FCanExecuteAction::CreateSP(this, &FPhysicsAssetEditor::IsNotSimulation));
 
 	ToolkitCommands->MapAction(
@@ -1341,7 +1341,7 @@ void FPhysicsAssetEditor::BuildMenuWidgetSelection(FMenuBuilder& InMenuBuilder)
 		InMenuBuilder.AddMenuEntry( Commands.ToggleShowSelected );
 		InMenuBuilder.AddMenuEntry( Commands.ShowSelected );
 		InMenuBuilder.AddMenuEntry( Commands.HideSelected );
-		InMenuBuilder.AddMenuEntry( Commands.ShowOnlySelected );
+		InMenuBuilder.AddMenuEntry( Commands.ToggleShowOnlySelected );
 		InMenuBuilder.AddMenuEntry( Commands.ShowAll );
 		InMenuBuilder.AddMenuEntry( Commands.HideAll );
 		InMenuBuilder.EndSection();
@@ -2826,33 +2826,32 @@ void FPhysicsAssetEditor::OnToggleShowSelected()
 	SharedData->ToggleShowSelected();
 }
 
-void FPhysicsAssetEditor::OnShowAll()
+void FPhysicsAssetEditor::OnShowSelected()
 {
-	SharedData->ShowAll();
+	SharedData->ShowSelected();
 }
 
-void FPhysicsAssetEditor::OnHideAll()
+void FPhysicsAssetEditor::OnHideSelected()
 {
-	SharedData->HideAll();
+	SharedData->HideSelected();
 }
 
-void FPhysicsAssetEditor::OnShowSelected()
+void FPhysicsAssetEditor::OnToggleShowOnlySelected()
 {
-	SharedData->ShowSelected();
+	SharedData->ToggleShowOnlySelected();
 }
 
-void FPhysicsAssetEditor::OnHideSelected()
+void FPhysicsAssetEditor::OnShowAll()
 {
-	SharedData->HideSelected();
+	SharedData->ShowAll();
 }
 
-
-void FPhysicsAssetEditor::OnShowOnlySelected()
+void FPhysicsAssetEditor::OnHideAll()
 {
 	SharedData->HideAll();
-	SharedData->ShowSelected();
 }
 
+
 void FPhysicsAssetEditor::OnDeselectAll()
 {
 	SharedData->ClearSelectedBody();
diff --git a/Engine/Source/Editor/PhysicsAssetEditor/Private/PhysicsAssetEditor.h b/Engine/Source/Editor/PhysicsAssetEditor/Private/PhysicsAssetEditor.h
--- a/Engine/Source/Editor/PhysicsAssetEditor/Private/PhysicsAssetEditor.h
+++ b/Engine/Source/Editor/PhysicsAssetEditor/Private/PhysicsAssetEditor.h
@@ -281,11 +281,11 @@ private:
 	void OnSelectAllConstraints();
 	void OnToggleSelectionType();
 	void OnToggleShowSelected();
-	void OnShowAll();
-	void OnHideAll();
 	void OnShowSelected();
 	void OnHideSelected();
-	void OnShowOnlySelected();
+	void OnToggleShowOnlySelected();
+	void OnShowAll();
+	void OnHideAll();
 	void OnDeselectAll();
 
 	FText GetRepeatLastSimulationToolTip() const;
diff --git a/Engine/Source/Editor/PhysicsAssetEditor/Private/PhysicsAssetEditorActions.cpp b/Engine/Source/Editor/PhysicsAssetEditor/Private/PhysicsAssetEditorActions.cpp
--- a/Engine/Source/Editor/PhysicsAssetEditor/Private/PhysicsAssetEditorActions.cpp
+++ b/Engine/Source/Editor/PhysicsAssetEditor/Private/PhysicsAssetEditorActions.cpp
@@ -83,7 +83,7 @@ void FPhysicsAssetEditorCommands::RegisterCommands()
 	UI_COMMAND(ToggleShowSelected, "Toggle Show Selected", "Show/hide selected bodies/constraints.", EUserInterfaceActionType::Button, FInputChord(EKeys::H));
 	UI_COMMAND(ShowSelected, "Show Selected", "Show selected bodies/constraints.", EUserInterfaceActionType::Button, FInputChord(EModifierKey::Shift, EKeys::H));
 	UI_COMMAND(HideSelected, "Hide Selected", "Hide selected bodies/constraints.", EUserInterfaceActionType::Button, FInputChord(EModifierKey::Control, EKeys::H));
-	UI_COMMAND(ShowOnlySelected, "Show Only Selected", "Show only the selected bodies/constraints, hiding all others.", EUserInterfaceActionType::Button, FInputChord(EKeys::G));
+	UI_COMMAND(ToggleShowOnlySelected, "Toggle Isolate Selected", "Show the selected bodies/constraints, hiding all others, or Show all.", EUserInterfaceActionType::Button, FInputChord(EKeys::G));
 	UI_COMMAND(ShowAll, "Show All", "Show all bodies/constraints.", EUserInterfaceActionType::Button, FInputChord(EModifierKey::Shift, EKeys::G));
 	UI_COMMAND(HideAll, "Hide All", "Hide all bodies/constraints.", EUserInterfaceActionType::Button, FInputChord(EModifierKey::Control, EKeys::G));
 	UI_COMMAND(DeselectAll, "Deselect All Bodies", "Deselect All", EUserInterfaceActionType::Button, FInputChord(EKeys::Escape));
diff --git a/Engine/Source/Editor/PhysicsAssetEditor/Private/PhysicsAssetEditorActions.h b/Engine/Source/Editor/PhysicsAssetEditor/Private/PhysicsAssetEditorActions.h
--- a/Engine/Source/Editor/PhysicsAssetEditor/Private/PhysicsAssetEditorActions.h
+++ b/Engine/Source/Editor/PhysicsAssetEditor/Private/PhysicsAssetEditorActions.h
@@ -95,7 +95,7 @@ public:
 	TSharedPtr<FUICommandInfo> ToggleShowSelected;
 	TSharedPtr<FUICommandInfo> ShowSelected;
 	TSharedPtr<FUICommandInfo> HideSelected;
-	TSharedPtr<FUICommandInfo> ShowOnlySelected;
+	TSharedPtr<FUICommandInfo> ToggleShowOnlySelected;
 	TSharedPtr<FUICommandInfo> ShowAll;
 	TSharedPtr<FUICommandInfo> HideAll;
 	TSharedPtr<FUICommandInfo> DeselectAll;
diff --git a/Engine/Source/Editor/PhysicsAssetEditor/Private/PhysicsAssetEditorSharedData.cpp b/Engine/Source/Editor/PhysicsAssetEditor/Private/PhysicsAssetEditorSharedData.cpp
--- a/Engine/Source/Editor/PhysicsAssetEditor/Private/PhysicsAssetEditorSharedData.cpp
+++ b/Engine/Source/Editor/PhysicsAssetEditor/Private/PhysicsAssetEditorSharedData.cpp
@@ -591,6 +591,78 @@ void FPhysicsAssetEditorSharedData::ToggleShowSelected()
 	}
 }
 
+void FPhysicsAssetEditorSharedData::ToggleShowOnlySelected()
+{
+	// Show only selected: make selected items visible and all others invisible.
+	// If we are already in the ShowOnlySelected state, make all visible.
+	bool bAllSelectedVisible = true;
+	if (bAllSelectedVisible)
+	{
+		for (const FSelection& Selection : SelectedConstraints)
+		{
+			if (HiddenConstraints.Contains(Selection.Index))
+			{
+				bAllSelectedVisible = false;
+				break;
+			}
+		}
+	}
+	if (bAllSelectedVisible)
+	{
+		for (const FSelection& Selection : SelectedBodies)
+		{
+			if (HiddenBodies.Contains(Selection.Index))
+			{
+				bAllSelectedVisible = false;
+			}
+		}
+	}
+
+	bool bAllNotSelectedHidden = true;
+	if (bAllNotSelectedHidden)
+	{
+		for (int32 ConstraintIndex = 0; ConstraintIndex < PhysicsAsset->ConstraintSetup.Num(); ++ConstraintIndex)
+		{
+			// Look at unselected constraints
+			if (!SelectedConstraints.ContainsByPredicate([ConstraintIndex](FSelection& V) { return V.Index == ConstraintIndex; } ))
+			{
+				// Is it hidden?
+				if (!HiddenConstraints.Contains(ConstraintIndex))
+				{
+					bAllNotSelectedHidden = false;
+					break;
+				}
+			}
+		}
+	}
+	if (bAllNotSelectedHidden)
+	{
+		for (int32 BodyIndex = 0; BodyIndex < PhysicsAsset->SkeletalBodySetups.Num(); ++BodyIndex)
+		{
+			// Look at unselected bodies
+			if (!SelectedBodies.ContainsByPredicate([BodyIndex](FSelection& V) { return V.Index == BodyIndex; }))
+			{
+				// Is it hidden?
+				if (!HiddenBodies.Contains(BodyIndex))
+				{
+					bAllNotSelectedHidden = false;
+					break;
+				}
+			}
+		}
+	}
+
+	if (bAllSelectedVisible && bAllNotSelectedHidden)
+	{
+		ShowAll();
+	}
+	else
+	{
+		HideAll();
+		ShowSelected();
+	}
+}
+
 void FPhysicsAssetEditorSharedData::ShowAll()
 {
 	HiddenConstraints.Empty();
diff --git a/Engine/Source/Editor/PhysicsAssetEditor/Private/PhysicsAssetEditorSharedData.h b/Engine/Source/Editor/PhysicsAssetEditor/Private/PhysicsAssetEditorSharedData.h
--- a/Engine/Source/Editor/PhysicsAssetEditor/Private/PhysicsAssetEditorSharedData.h
+++ b/Engine/Source/Editor/PhysicsAssetEditor/Private/PhysicsAssetEditorSharedData.h
@@ -147,6 +147,7 @@ public:
 	void ToggleShowSelected();
 	void ShowAll();
 	void HideAll();
+	void ToggleShowOnlySelected();
 	void ShowSelected();
 	void HideSelected();
 	void SetSelectedBodyAnyPrim(int32 BodyIndex, bool bSelected);
