PhAT - improve ToggleShowVisible when changing selections

diff --git a/Engine/Source/Editor/PhysicsAssetEditor/Private/PhysicsAssetEditorSharedData.cpp b/Engine/Source/Editor/PhysicsAssetEditor/Private/PhysicsAssetEditorSharedData.cpp
--- a/Engine/Source/Editor/PhysicsAssetEditor/Private/PhysicsAssetEditorSharedData.cpp
+++ b/Engine/Source/Editor/PhysicsAssetEditor/Private/PhysicsAssetEditorSharedData.cpp
@@ -567,28 +567,37 @@ void FPhysicsAssetEditorSharedData::ToggleSelectionType()
 
 void FPhysicsAssetEditorSharedData::ToggleShowSelected()
 {
-	for (const FSelection& Selection : SelectedConstraints)
+	bool bAllSelectedVisible = true;
+	if (bAllSelectedVisible)
 	{
-		if (HiddenConstraints.Contains(Selection.Index))
-		{
-			HiddenConstraints.Remove(Selection.Index);
-		}
-		else
+		for (const FSelection& Selection : SelectedConstraints)
 		{
-			HiddenConstraints.AddUnique(Selection.Index);
+			if (HiddenConstraints.Contains(Selection.Index))
+			{
+				bAllSelectedVisible = false;
+				break;
+			}
 		}
 	}
-	for (const FSelection& Selection : SelectedBodies)
+	if (bAllSelectedVisible)
 	{
-		if (HiddenBodies.Contains(Selection.Index))
-		{
-			HiddenBodies.Remove(Selection.Index);
-		}
-		else
+		for (const FSelection& Selection : SelectedBodies)
 		{
-			HiddenBodies.AddUnique(Selection.Index);
+			if (HiddenBodies.Contains(Selection.Index))
+			{
+				bAllSelectedVisible = false;
+			}
 		}
 	}
+
+	if (bAllSelectedVisible)
+	{
+		HideSelected();
+	}
+	else
+	{
+		ShowSelected();
+	}
 }
 
 void FPhysicsAssetEditorSharedData::ToggleShowOnlySelected()
