Chaos Physics - PHAT editor
Properly set the skeletalMeshCOmponent for the PhysicalAnimationComponent in RBAN mode

diff --git a/Engine/Source/Editor/PhysicsAssetEditor/Private/PhysicsAssetEditorSharedData.cpp b/Engine/Source/Editor/PhysicsAssetEditor/Private/PhysicsAssetEditorSharedData.cpp
--- a/Engine/Source/Editor/PhysicsAssetEditor/Private/PhysicsAssetEditorSharedData.cpp
+++ b/Engine/Source/Editor/PhysicsAssetEditor/Private/PhysicsAssetEditorSharedData.cpp
@@ -1922,7 +1922,7 @@ void FPhysicsAssetEditorSharedData::EnableSimulation(bool bEnableSimulation)
 		// Enable the PreviewInstance (containing the AnimNode_RigidBody)
 		EditorSkelComp->SetAnimationMode(EAnimationMode::AnimationCustomMode);
 		EditorSkelComp->InitAnim(true);
-
+		PhysicalAnimationComponent->SetSkeletalMeshComponent(EditorSkelComp);
 		// Add the floor
 		TSharedPtr<IPersonaPreviewScene> Scene = PreviewScene.Pin();
 		if (Scene != nullptr)
