Physics Asset Editor - Fix selection synchronization (select all does not select multiple primitive bodies )
- Simplify and fix the way skeleton tree is synchronized with selection ( also made it faster for large number of bones/bodies )

diff --git a/Engine/Source/Editor/PhysicsAssetEditor/Private/PhysicsAssetEditor.cpp b/Engine/Source/Editor/PhysicsAssetEditor/Private/PhysicsAssetEditor.cpp
--- a/Engine/Source/Editor/PhysicsAssetEditor/Private/PhysicsAssetEditor.cpp
+++ b/Engine/Source/Editor/PhysicsAssetEditor/Private/PhysicsAssetEditor.cpp
@@ -94,6 +94,31 @@ namespace PhysicsAssetEditor
 {
 	const float	DefaultPrimSize = 15.0f;
 	const float	DuplicateXOffset = 10.0f;
+
+	/** contain everything to identify a shape uniquely - used for synchronizing selection mostly */
+	struct FShapeData
+	{
+		int32 Index;
+		int32 PrimitiveIndex;
+
+		FShapeData(int32 Index, int32 PrimitiveIndex)
+			: Index(Index)
+			, PrimitiveIndex(PrimitiveIndex)
+		{
+		}
+
+		bool operator==(const FShapeData& rhs) const
+		{
+			return Index == rhs.Index && PrimitiveIndex == rhs.PrimitiveIndex;
+		}
+
+		friend uint32 GetTypeHash(const FShapeData& ShapeData)
+		{
+			return HashCombine(GetTypeHash(ShapeData.Index), GetTypeHash(ShapeData.PrimitiveIndex));
+		}
+	};
+
+	
 }
 
 void FPhysicsAssetEditor::RegisterTabSpawners(const TSharedRef<class FTabManager>& InTabManager)
@@ -238,9 +263,13 @@ void FPhysicsAssetEditor::HandleViewportSelectionChanged(const TArray<FPhysicsAs
 		}
 		else
 		{
-			TArray<UObject*> Objects;
+
+			// let's store all the selection in sets so that when we go through the list of items in the list 
+			// ( which can be long ) we only do O(1) lookup for each of them 
+			TSet<UObject*> Objects;
 			TSet<USkeletalBodySetup*> Bodies;
 			TSet<UPhysicsConstraintTemplate*> Constraints;
+			TSet<PhysicsAssetEditor::FShapeData> Shapes;
 			Algo::Transform(InSelectedBodies, Objects, [this](const FPhysicsAssetEditorSharedData::FSelection& InItem) 
 			{ 
 				return SharedData->PhysicsAsset->SkeletalBodySetups[InItem.Index];
@@ -257,53 +286,48 @@ void FPhysicsAssetEditor::HandleViewportSelectionChanged(const TArray<FPhysicsAs
 			{ 
 				return SharedData->PhysicsAsset->ConstraintSetup[InItem.Index];
 			});
+			Algo::Transform(InSelectedConstraints, Shapes, [this](const FPhysicsAssetEditorSharedData::FSelection& InItem)
+			{
+				return PhysicsAssetEditor::FShapeData(InItem.Index, InItem.PrimitiveIndex);
+			});
 
 			if (PhysAssetProperties.IsValid())
 			{
-				PhysAssetProperties->SetObjects(Objects);
+				PhysAssetProperties->SetObjects(Objects.Array());
 			}
 
 			if (SkeletonTree.IsValid())
 			{
-				SkeletonTree->SelectItemsBy([this, &InSelectedBodies, &Constraints](const TSharedRef<ISkeletonTreeItem>& InItem, bool& bInOutExpand)
+				SkeletonTree->SelectItemsBy([this, &Objects, &Constraints, &Bodies, &Shapes](const TSharedRef<ISkeletonTreeItem>& InItem, bool& bInOutExpand)
 				{
-					if(InItem->IsOfType<FSkeletonTreePhysicsBodyItem>() || InItem->IsOfType<FSkeletonTreePhysicsShapeItem>())
+					if (InItem->IsOfType<FSkeletonTreePhysicsBodyItem>())
 					{
-						for (const FPhysicsAssetEditorSharedData::FSelection& SelectedBody : InSelectedBodies)
+						const USkeletalBodySetup* BodySetup = Cast<USkeletalBodySetup>(InItem->GetObject());
+						if (Bodies.Contains(BodySetup))
 						{
-							USkeletalBodySetup* BodySetup = SharedData->PhysicsAsset->SkeletalBodySetups[SelectedBody.Index];
-							int32 GeomCount = BodySetup->AggGeom.GetElementCount();
-							if (BodySetup == InItem->GetObject())
-							{
-								if(InItem->IsOfType<FSkeletonTreePhysicsShapeItem>())
-								{
-									TSharedRef<FSkeletonTreePhysicsShapeItem> SkeletonTreePhysicsShapeItem = StaticCastSharedRef<FSkeletonTreePhysicsShapeItem>(InItem);
-									if(SkeletonTreePhysicsShapeItem->GetShapeIndex() == SelectedBody.PrimitiveIndex && SkeletonTreePhysicsShapeItem->GetShapeType() == SelectedBody.PrimitiveType && GeomCount > 1)
-									{
-										bInOutExpand = true;
-										return true;
-									}
-								}
-								else if(GeomCount <= 1)
-								{
-									bInOutExpand = true;
-									return true;
-								}
-							}
+							bInOutExpand = true;
+							return true;
 						}
 					}
-					else if(InItem->IsOfType<FSkeletonTreePhysicsConstraintItem>())
+					else if (InItem->IsOfType<FSkeletonTreePhysicsShapeItem>())
 					{
-						for (UPhysicsConstraintTemplate* Constraint : Constraints)
+						TSharedRef<FSkeletonTreePhysicsShapeItem> ShapeItem = StaticCastSharedRef<FSkeletonTreePhysicsShapeItem>(InItem);
+						PhysicsAssetEditor::FShapeData ShapeData(ShapeItem->GetBodySetupIndex(), ShapeItem->GetShapeIndex());
+						if (Shapes.Contains(ShapeData))
 						{
-							if (Constraint == InItem->GetObject())
-							{
-								bInOutExpand = true;
-								return true;
-							}
+							bInOutExpand = true;
+							return true;
+						}
+					}
+					else if (InItem->IsOfType<FSkeletonTreePhysicsConstraintItem>())
+					{
+						const UPhysicsConstraintTemplate* Constraint = Cast<UPhysicsConstraintTemplate>(InItem->GetObject());
+						if (Constraints.Contains(Constraint))
+						{
+							bInOutExpand = true;
+							return true;
 						}
 					}
-
 					return false;
 				});
 			}
