Build fix : tentative fix for build break

diff --git a/Engine/Source/Editor/PhysicsAssetEditor/Private/PhysicsAssetEditor.cpp b/Engine/Source/Editor/PhysicsAssetEditor/Private/PhysicsAssetEditor.cpp
--- a/Engine/Source/Editor/PhysicsAssetEditor/Private/PhysicsAssetEditor.cpp
+++ b/Engine/Source/Editor/PhysicsAssetEditor/Private/PhysicsAssetEditor.cpp
@@ -24,6 +24,7 @@
 #include "ScopedTransaction.h"
 #include "PhysicsAssetEditorActions.h"
 #include "PhysicsAssetEditorSkeletalMeshComponent.h"
+#include "Templates/TypeHash.h"
 
 #include "PropertyEditorModule.h"
 #include "IDetailsView.h"
@@ -114,7 +115,7 @@ namespace PhysicsAssetEditor
 
 		friend uint32 GetTypeHash(const FShapeData& ShapeData)
 		{
-			return HashCombine(GetTypeHash(ShapeData.Index), GetTypeHash(ShapeData.PrimitiveIndex));
+			return HashCombine(::GetTypeHash(ShapeData.Index), ::GetTypeHash(ShapeData.PrimitiveIndex));
 		}
 	};
 
