Physics Asset Editor - Allow bone selection to show the "Copy Collision From StaticMesh" contextual menu
Also add solid debug draw in preview for convex volumes

diff --git a/Engine/Source/Editor/PhysicsAssetEditor/Private/PhysicsAssetEditor.cpp b/Engine/Source/Editor/PhysicsAssetEditor/Private/PhysicsAssetEditor.cpp
--- a/Engine/Source/Editor/PhysicsAssetEditor/Private/PhysicsAssetEditor.cpp
+++ b/Engine/Source/Editor/PhysicsAssetEditor/Private/PhysicsAssetEditor.cpp
@@ -1194,6 +1194,20 @@ void FPhysicsAssetEditor::Mirror()
 	RefreshPreviewViewport();
 }
 
+void FPhysicsAssetEditor::AddAdvancedMenuWidget(FMenuBuilder& InMenuBuilder)
+{
+	InMenuBuilder.BeginSection("Advanced", LOCTEXT("AdvancedHeading", "Advanced"));
+	InMenuBuilder.AddSubMenu(
+		LOCTEXT("AddCollisionfromStaticMesh", "Copy Collision From StaticMesh"),
+		LOCTEXT("AddCollisionfromStaticMesh_Tooltip", "Copy convex collision from a specified static mesh"),
+		FNewMenuDelegate::CreateLambda([this](FMenuBuilder& InSubMenuBuilder)
+			{
+				InSubMenuBuilder.AddWidget(BuildStaticMeshAssetPicker(), FText(), true);
+			})
+	);
+	InMenuBuilder.EndSection();
+}
+
 void FPhysicsAssetEditor::BuildMenuWidgetBody(FMenuBuilder& InMenuBuilder)
 {
 	InMenuBuilder.PushCommandList(GetToolkitCommands());
@@ -1284,16 +1298,7 @@ void FPhysicsAssetEditor::BuildMenuWidgetBody(FMenuBuilder& InMenuBuilder)
 		InMenuBuilder.AddMenuEntry( Commands.RemoveBodyFromPhysicalAnimationProfile );
 		InMenuBuilder.EndSection();
 
-		InMenuBuilder.BeginSection( "Advanced", LOCTEXT( "AdvancedHeading", "Advanced" ) );
-		InMenuBuilder.AddSubMenu(
-			LOCTEXT("AddCollisionfromStaticMesh", "Copy Collision From StaticMesh"), 
-			LOCTEXT("AddCollisionfromStaticMesh_Tooltip", "Copy convex collision from a specified static mesh"), 
-			FNewMenuDelegate::CreateLambda([this](FMenuBuilder& InSubMenuBuilder)
-			{
-				InSubMenuBuilder.AddWidget(BuildStaticMeshAssetPicker(), FText(), true);
-			})
-		);
-		InMenuBuilder.EndSection();
+		AddAdvancedMenuWidget(InMenuBuilder);
 	}
 	InMenuBuilder.PopCommandList();
 }
@@ -1475,6 +1480,7 @@ void FPhysicsAssetEditor::BuildMenuWidgetBone(FMenuBuilder& InMenuBuilder)
 			FNewMenuDelegate::CreateStatic( &FillAddShapeMenu ) );
 	}
 	InMenuBuilder.EndSection();
+	AddAdvancedMenuWidget(InMenuBuilder);
 	InMenuBuilder.PopCommandList();
 }
 
@@ -2600,7 +2606,36 @@ void FPhysicsAssetEditor::OnAssetSelectedFromStaticMeshAssetPicker( const FAsset
 	// Make sure rendering is done - so we are not changing data being used by collision drawing.
 	FlushRenderingCommands();
 
-	if (SharedData->GetSelectedBody())
+	// get select bones
+	TArray<TSharedPtr<ISkeletonTreeItem>> Items = SkeletonTree->GetSelectedItems();
+	FSkeletonTreeSelection Selection(Items);
+	TArray<TSharedPtr<ISkeletonTreeItem>> BoneItems = Selection.GetSelectedItemsByTypeId("FSkeletonTreeBoneItem");
+
+	// gather all the body indices from both the body and bone selection 
+	// make sure to create a body setup if we encounter a bone with no associated body
+	TSet<int32> BodyIndicesToUpdate;
+	if (SharedData->GetSelectedBody() || BoneItems.Num() > 0)
+	{
+		for (int32 SelectedBodyIndex = 0; SelectedBodyIndex < SharedData->SelectedBodies.Num(); ++SelectedBodyIndex)
+		{
+			BodyIndicesToUpdate.Add(SharedData->SelectedBodies[SelectedBodyIndex].Index);
+		}
+
+		for (TSharedPtr<ISkeletonTreeItem> BoneItem : BoneItems)
+		{
+			UBoneProxy* BoneProxy = CastChecked<UBoneProxy>(BoneItem->GetObject());
+			int32 BodyIndex = SharedData->PhysicsAsset->FindBodyIndex(BoneProxy->BoneName);
+			if (BodyIndex == INDEX_NONE)
+			{
+				// no associated body found, let's create one
+				const FPhysAssetCreateParams& NewBodyData = GetDefault<UPhysicsAssetGenerationSettings>()->CreateParams;
+				BodyIndex = FPhysicsAssetUtils::CreateNewBody(SharedData->PhysicsAsset, BoneProxy->BoneName, NewBodyData);
+			}
+			BodyIndicesToUpdate.Add(BodyIndex);
+		}
+	}
+
+	if (BodyIndicesToUpdate.Num() > 0)
 	{
 		UStaticMesh* SM = Cast<UStaticMesh>(AssetData.GetAsset());
 
@@ -2608,9 +2643,9 @@ void FPhysicsAssetEditor::OnAssetSelectedFromStaticMeshAssetPicker( const FAsset
 		{
 			SharedData->PhysicsAsset->Modify();
 
-			for (int32 SelectedBodyIndex = 0; SelectedBodyIndex < SharedData->SelectedBodies.Num(); ++SelectedBodyIndex)
+			for (int32 BodyIndex: BodyIndicesToUpdate)
 			{
-				UBodySetup* BaseSetup = SharedData->PhysicsAsset->SkeletalBodySetups[SharedData->SelectedBodies[SelectedBodyIndex].Index];
+				UBodySetup* BaseSetup = SharedData->PhysicsAsset->SkeletalBodySetups[BodyIndex];
 				BaseSetup->Modify();
 				BaseSetup->AddCollisionFrom(SM->GetBodySetup());
 				BaseSetup->InvalidatePhysicsData();
diff --git a/Engine/Source/Editor/PhysicsAssetEditor/Private/PhysicsAssetEditor.h b/Engine/Source/Editor/PhysicsAssetEditor/Private/PhysicsAssetEditor.h
--- a/Engine/Source/Editor/PhysicsAssetEditor/Private/PhysicsAssetEditor.h
+++ b/Engine/Source/Editor/PhysicsAssetEditor/Private/PhysicsAssetEditor.h
@@ -88,6 +88,7 @@ public:
 	TSharedRef<ISkeletonTree> BuildMenuWidgetNewConstraintForBody(FMenuBuilder& InMenuBuilder, int32 InSourceBodyIndex, SGraphEditor::FActionMenuClosed InOnActionMenuClosed = SGraphEditor::FActionMenuClosed());
 	void BuildMenuWidgetBone(FMenuBuilder& InMenuBuilder);
 	TSharedRef<SWidget> BuildStaticMeshAssetPicker();
+	void AddAdvancedMenuWidget(FMenuBuilder& InMenuBuilder);
 
 	/** IToolkit interface */
 	virtual FName GetToolkitFName() const override;
diff --git a/Engine/Source/Editor/PhysicsAssetEditor/Private/PhysicsAssetEditorSkeletalMeshComponent.cpp b/Engine/Source/Editor/PhysicsAssetEditor/Private/PhysicsAssetEditorSkeletalMeshComponent.cpp
--- a/Engine/Source/Editor/PhysicsAssetEditor/Private/PhysicsAssetEditorSkeletalMeshComponent.cpp
+++ b/Engine/Source/Editor/PhysicsAssetEditor/Private/PhysicsAssetEditorSkeletalMeshComponent.cpp
@@ -213,7 +213,12 @@ void UPhysicsAssetEditorSkeletalMeshComponent::RenderAssetTools(const FSceneView
 
 				FTransform ElemTM = GetPrimitiveTransform(BoneTM, i, EAggCollisionShape::Convex, j, Scale);
 
-				//convex doesn't have solid draw so render lines if we're in hitTestAndBodyMode
+				if (CollisionViewMode == EPhysicsAssetEditorRenderMode::Solid)
+				{
+					UMaterialInterface* PrimMaterial = GetPrimitiveMaterial(i, EAggCollisionShape::Sphyl, j);
+					AggGeom->ConvexElems[j].DrawElemSolid(PDI, ElemTM, Scale, PrimMaterial->GetRenderProxy());
+				}
+
 				if (CollisionViewMode == EPhysicsAssetEditorRenderMode::Solid || CollisionViewMode == EPhysicsAssetEditorRenderMode::Wireframe)
 				{
 					AggGeom->ConvexElems[j].DrawElemWire(PDI, ElemTM, Scale, GetPrimitiveColor(i, EAggCollisionShape::Convex, j));
diff --git a/Engine/Source/Runtime/Engine/Classes/PhysicsEngine/ConvexElem.h b/Engine/Source/Runtime/Engine/Classes/PhysicsEngine/ConvexElem.h
--- a/Engine/Source/Runtime/Engine/Classes/PhysicsEngine/ConvexElem.h
+++ b/Engine/Source/Runtime/Engine/Classes/PhysicsEngine/ConvexElem.h
@@ -72,6 +72,8 @@ public:
 
 	ENGINE_API void	DrawElemWire(class FPrimitiveDrawInterface* PDI, const FTransform& ElemTM, const float Scale, const FColor Color) const;
 
+	ENGINE_API void DrawElemSolid(FPrimitiveDrawInterface* PDI, const FTransform& ElemTM, const float Scale, const FMaterialRenderProxy* MaterialRenderProxy) const;
+
 	ENGINE_API void AddCachedSolidConvexGeom(TArray<FDynamicMeshVertex>& VertexBuffer, TArray<uint32>& IndexBuffer, const FColor VertexColor) const;
 
 	/** Reset the hull to empty all arrays */
diff --git a/Engine/Source/Runtime/Engine/Private/PhysicsEngine/PhysDrawing.cpp b/Engine/Source/Runtime/Engine/Private/PhysicsEngine/PhysDrawing.cpp
--- a/Engine/Source/Runtime/Engine/Private/PhysicsEngine/PhysDrawing.cpp
+++ b/Engine/Source/Runtime/Engine/Private/PhysicsEngine/PhysDrawing.cpp
@@ -717,11 +717,50 @@ void FKConvexElem::DrawElemWire(FPrimitiveDrawInterface* PDI, const FTransform&
 	}
 	else
 	{
-		UE_LOG(LogPhysics, Log, TEXT("FKConvexElem::AddCachedSolidConvexGeom : No ConvexMesh, so unable to draw."));
+		UE_LOG(LogPhysics, Log, TEXT("FKConvexElem::DrawElemWire : No ConvexMesh, so unable to draw."));
 	}
 #endif // WITH_PHYSX
 }
 
+void FKConvexElem::DrawElemSolid(FPrimitiveDrawInterface* PDI, const FTransform& ElemTM, const float Scale, const FMaterialRenderProxy* MaterialRenderProxy) const
+{
+#if WITH_CHAOS
+	const int32 NumIndices = IndexData.Num();
+	if (NumIndices > 0 && ensure(NumIndices % 3 == 0))
+	{
+		FDynamicMeshBuilder MeshBuilder(PDI->View->GetFeatureLevel());
+
+		const FVector2D DummyUV(0);
+		const FVector DummyTangentX(1, 0, 0);
+		const FVector DummyTangentY(0, 1, 0);
+		const FVector DummyTangentZ(0, 0, 1);
+
+		// NOTE: With chaos, instead of using a mesh with transformed verts, we use the 
+		// VertexData directly, so we don't need to remove the body->elem transform like
+		// we did with physx.
+		const int32 NumVerts = VertexData.Num();
+		for (const FVector& Vert : VertexData)
+		{
+			MeshBuilder.AddVertex(ElemTM.TransformPosition(Vert), DummyUV, DummyTangentX, DummyTangentY, DummyTangentZ, FColor::White);
+		}
+
+		for (int32 Base = 0; Base < NumIndices; Base += 3)
+		{
+			if (IndexData[Base] >= NumVerts || IndexData[Base + 1] >= NumVerts || IndexData[Base + 2] >= NumVerts)
+			{
+				continue;
+			}
+			MeshBuilder.AddTriangle(IndexData[Base], IndexData[Base + 1], IndexData[Base + 2]);
+		}
+		MeshBuilder.Draw(PDI, FMatrix::Identity, MaterialRenderProxy, SDPG_World, 0.f);
+	}
+	else
+	{
+		UE_LOG(LogPhysics, Log, TEXT("FKConvexElem::DrawElemSolid : No ConvexMesh, so unable to draw."));
+	}
+#endif
+}
+
 void FKConvexElem::AddCachedSolidConvexGeom(TArray<FDynamicMeshVertex>& VertexBuffer, TArray<uint32>& IndexBuffer, const FColor VertexColor) const
 {
 #if WITH_PHYSX && PHYSICS_INTERFACE_PHYSX
