Fix UE4 reference in comment for phys asset shared data

diff --git a/Engine/Source/Editor/PhysicsAssetEditor/Private/PhysicsAssetEditorSharedData.cpp b/Engine/Source/Editor/PhysicsAssetEditor/Private/PhysicsAssetEditorSharedData.cpp
--- a/Engine/Source/Editor/PhysicsAssetEditor/Private/PhysicsAssetEditorSharedData.cpp
+++ b/Engine/Source/Editor/PhysicsAssetEditor/Private/PhysicsAssetEditorSharedData.cpp
@@ -166,7 +166,7 @@ void FPhysicsAssetEditorSharedData::Initialize(const TSharedRef<IPersonaPreviewS
 				PhysicsAsset->SetFlags(RF_Transactional);
 				PhysicsAsset->Modify();
 
-				// Iterate backwards, as PhysicsAsset->SkeletalBodySetups is a TArray and UE4 containers don't support remove_if()
+				// Iterate backwards, as PhysicsAsset->SkeletalBodySetups is a TArray and Unreal containers don't support remove_if()
 				for ( int32 i = MissingBodyIndices.Num() - 1; i >= 0; --i )
 				{
 					DeleteBody( MissingBodyIndices[i], false );
