Disable main solver when running RBAN in PhAT.

diff --git a/Engine/Source/Editor/PhysicsAssetEditor/Private/PhysicsAssetEditorEditMode.cpp b/Engine/Source/Editor/PhysicsAssetEditor/Private/PhysicsAssetEditorEditMode.cpp
--- a/Engine/Source/Editor/PhysicsAssetEditor/Private/PhysicsAssetEditorEditMode.cpp
+++ b/Engine/Source/Editor/PhysicsAssetEditor/Private/PhysicsAssetEditorEditMode.cpp
@@ -486,7 +486,12 @@ void FPhysicsAssetEditorEditMode::Tick(FEditorViewportClient* ViewportClient, fl
 
 		// We back up the transforms array now
 		SharedData->EditorSkelComp->AnimationSpaceBases = SharedData->EditorSkelComp->GetComponentSpaceTransforms();
-		SharedData->EditorSkelComp->SetPhysicsBlendWeight(SharedData->EditorOptions->PhysicsBlend);
+		// When using the World solver, we must specify how much of the solver output gets blended into the animated mesh pose
+		// When using other solvers in PhAT, we don't want SetPhysicsBlendWeight function to re-enale the main solver physics
+		if (SharedData->RunningSimulationSolverType == EPhysicsAssetEditorSolverType::World)
+		{
+			SharedData->EditorSkelComp->SetPhysicsBlendWeight(SharedData->EditorOptions->PhysicsBlend);
+		}
 		SharedData->EditorSkelComp->bUpdateJointsFromAnimation = SharedData->EditorOptions->bUpdateJointsFromAnimation;
 		SharedData->EditorSkelComp->PhysicsTransformUpdateMode = SharedData->EditorOptions->PhysicsUpdateMode;
 
diff --git a/Engine/Source/Editor/PhysicsAssetEditor/Private/PhysicsAssetEditorSharedData.cpp b/Engine/Source/Editor/PhysicsAssetEditor/Private/PhysicsAssetEditorSharedData.cpp
--- a/Engine/Source/Editor/PhysicsAssetEditor/Private/PhysicsAssetEditorSharedData.cpp
+++ b/Engine/Source/Editor/PhysicsAssetEditor/Private/PhysicsAssetEditorSharedData.cpp
@@ -1931,6 +1931,9 @@ void FPhysicsAssetEditorSharedData::EnableSimulation(bool bEnableSimulation)
 			EditorSkelComp->SetAnimationMode(EAnimationMode::AnimationCustomMode);
 			EditorSkelComp->InitAnim(true);
 
+			// Disable main solver physics
+			EditorSkelComp->SetAllBodiesSimulatePhysics(false);
+
 			// Add the floor
 			TSharedPtr<IPersonaPreviewScene> Scene = PreviewScene.Pin();
 			if (Scene != nullptr)
