Move solver type option to UPhysicsAsset and set default solver type to World in PhAT

diff --git a/Engine/Source/Editor/PhysicsAssetEditor/Private/PhysicsAssetEditorEditMode.cpp b/Engine/Source/Editor/PhysicsAssetEditor/Private/PhysicsAssetEditorEditMode.cpp
--- a/Engine/Source/Editor/PhysicsAssetEditor/Private/PhysicsAssetEditorEditMode.cpp
+++ b/Engine/Source/Editor/PhysicsAssetEditor/Private/PhysicsAssetEditorEditMode.cpp
@@ -488,7 +488,7 @@ void FPhysicsAssetEditorEditMode::Tick(FEditorViewportClient* ViewportClient, fl
 		SharedData->EditorSkelComp->AnimationSpaceBases = SharedData->EditorSkelComp->GetComponentSpaceTransforms();
 		// When using the World solver, we must specify how much of the solver output gets blended into the animated mesh pose
 		// When using other solvers in PhAT, we don't want SetPhysicsBlendWeight function to re-enale the main solver physics
-		if (SharedData->RunningSimulationSolverType == EPhysicsAssetEditorSolverType::World)
+		if (SharedData->PhysicsAsset->SolverType == EPhysicsAssetSolverType::World)
 		{
 			SharedData->EditorSkelComp->SetPhysicsBlendWeight(SharedData->EditorOptions->PhysicsBlend);
 		}
diff --git a/Engine/Source/Editor/PhysicsAssetEditor/Private/PhysicsAssetEditorSharedData.cpp b/Engine/Source/Editor/PhysicsAssetEditor/Private/PhysicsAssetEditorSharedData.cpp
--- a/Engine/Source/Editor/PhysicsAssetEditor/Private/PhysicsAssetEditorSharedData.cpp
+++ b/Engine/Source/Editor/PhysicsAssetEditor/Private/PhysicsAssetEditorSharedData.cpp
@@ -81,8 +81,6 @@ FPhysicsAssetEditorSharedData::FPhysicsAssetEditorSharedData()
 	check(EditorOptions);
 
 	EditorOptions->LoadConfig();
-
-	RunningSimulationSolverType = EditorOptions->SolverType;
 }
 
 FPhysicsAssetEditorSharedData::~FPhysicsAssetEditorSharedData()
@@ -1906,11 +1904,8 @@ void FPhysicsAssetEditorSharedData::EnableSimulation(bool bEnableSimulation)
 {
 	if (bEnableSimulation)
 	{
-		// save the current solver type to make sure we clean up using the one we started the simulation with
-		RunningSimulationSolverType = EditorOptions->SolverType;
-
 		// in Chaos, we have to manipulate the RBAN node in the Anim Instance (at least until we get SkelMeshComp implemented)
-		const bool bUseRBANSolver = (RunningSimulationSolverType == EPhysicsAssetEditorSolverType::RBAN);
+		const bool bUseRBANSolver = (PhysicsAsset->SolverType == EPhysicsAssetSolverType::RBAN);
 		MouseHandle->SetAnimInstanceMode(bUseRBANSolver);
 
 		if (!bUseRBANSolver)
diff --git a/Engine/Source/Editor/PhysicsAssetEditor/Private/PhysicsAssetEditorSharedData.h b/Engine/Source/Editor/PhysicsAssetEditor/Private/PhysicsAssetEditorSharedData.h
--- a/Engine/Source/Editor/PhysicsAssetEditor/Private/PhysicsAssetEditorSharedData.h
+++ b/Engine/Source/Editor/PhysicsAssetEditor/Private/PhysicsAssetEditorSharedData.h
@@ -313,9 +313,6 @@ public:
 	bool bRunningSimulation;
 	bool bNoGravitySimulation;
 	
-	/** solver type used while simulation is running */
-	EPhysicsAssetEditorSolverType RunningSimulationSolverType;
-
 	/** Manipulation (rotate, translate, scale) */
 	bool bManipulating;
 
diff --git a/Engine/Source/Editor/UnrealEd/Classes/Preferences/PhysicsAssetEditorOptions.h b/Engine/Source/Editor/UnrealEd/Classes/Preferences/PhysicsAssetEditorOptions.h
--- a/Engine/Source/Editor/UnrealEd/Classes/Preferences/PhysicsAssetEditorOptions.h
+++ b/Engine/Source/Editor/UnrealEd/Classes/Preferences/PhysicsAssetEditorOptions.h
@@ -25,13 +25,6 @@ enum class EPhysicsAssetEditorConstraintViewMode : uint8
 	AllLimits
 };
 
-UENUM()
-enum class EPhysicsAssetEditorSolverType: uint8
-{
-	RBAN,
-	World,
-};
-
 UCLASS(hidecategories=Object, config=EditorPerProjectUserSettings)
 class UNREALED_API UPhysicsAssetEditorOptions : public UObject
 {
@@ -152,7 +145,4 @@ class UNREALED_API UPhysicsAssetEditorOptions : public UObject
 	/** When set, cloth will reset each time simulation is toggled */
 	UPROPERTY(EditAnywhere, config, Category=Clothing)
 	bool bResetClothWhenSimulating;
-
-	UPROPERTY(EditAnywhere, config, Category = Simulation)
-	EPhysicsAssetEditorSolverType SolverType;
 };
diff --git a/Engine/Source/Runtime/Engine/Classes/PhysicsEngine/PhysicsAsset.h b/Engine/Source/Runtime/Engine/Classes/PhysicsEngine/PhysicsAsset.h
--- a/Engine/Source/Runtime/Engine/Classes/PhysicsEngine/PhysicsAsset.h
+++ b/Engine/Source/Runtime/Engine/Classes/PhysicsEngine/PhysicsAsset.h
@@ -77,6 +77,13 @@ struct ENGINE_API FSolverIterations
 		int32 CollisionPushOutIterations;
 };
 
+UENUM()
+enum class EPhysicsAssetSolverType: uint8
+{
+	RBAN,
+	World,
+};
+
 /**
  * PhysicsAsset contains a set of rigid bodies and constraints that make up a single ragdoll.
  * The asset is not limited to human ragdolls, and can be used for any physical simulation using bodies and constraints.
@@ -141,6 +148,10 @@ public:
 	UPROPERTY(EditAnywhere, AdvancedDisplay, Category = SolverSettings)
 	FSolverIterations SolverIterations;
 
+	/** [Chaos Only] Solver type used in physics asset editor. */
+	UPROPERTY(EditAnywhere, AdvancedDisplay, Category = SolverSettings)
+	EPhysicsAssetSolverType SolverType;
+
 	/** If true, we skip instancing bodies for this PhysicsAsset on dedicated servers */
 	UPROPERTY(EditAnywhere, AdvancedDisplay, Category = Physics)
 	uint8 bNotForDedicatedServer:1;
diff --git a/Engine/Source/Runtime/Engine/Private/PhysicsEngine/PhysicsAsset.cpp b/Engine/Source/Runtime/Engine/Private/PhysicsEngine/PhysicsAsset.cpp
--- a/Engine/Source/Runtime/Engine/Private/PhysicsEngine/PhysicsAsset.cpp
+++ b/Engine/Source/Runtime/Engine/Private/PhysicsEngine/PhysicsAsset.cpp
@@ -22,6 +22,9 @@
 
 #define LOCTEXT_NAMESPACE "PhysicsAsset"
 
+bool bUseRBANForDefaultPhysicsAssetSolverType = false;
+FAutoConsoleVariableRef CVarUseRBANForDefaultPhysicsAssetSolverType(TEXT("p.Chaos.UseRBANForDefaultPhysicsAssetSolverType"), bUseRBANForDefaultPhysicsAssetSolverType, TEXT("Boolean to use RBAN for default physics asset solver type (false by default)"));
+
 FSolverIterations::FSolverIterations()
 	: FixedTimeStep(0)
 	, SolverIterations(3)
@@ -40,6 +43,14 @@ FSolverIterations::FSolverIterations()
 UPhysicsAsset::UPhysicsAsset(const FObjectInitializer& ObjectInitializer)
 	: Super(ObjectInitializer)
 {
+	if(bUseRBANForDefaultPhysicsAssetSolverType)
+	{
+		SolverType = EPhysicsAssetSolverType::RBAN;
+	}
+	else
+	{
+		SolverType = EPhysicsAssetSolverType::World;
+	}
 }
 
 void UPhysicsAsset::UpdateBoundsBodiesArray()
