Physics Asset editor - fix simulation stopping when clicking in the 3D viewport after editing atext field and starting the simulation
Issue was caused by the user focus been kept by the detail window and when clicking in the 3D viewport , textfield would get a lost focus eventand trigger OnFinishedChangingProperties

diff --git a/Engine/Source/Editor/PhysicsAssetEditor/Private/PhysicsAssetEditor.cpp b/Engine/Source/Editor/PhysicsAssetEditor/Private/PhysicsAssetEditor.cpp
--- a/Engine/Source/Editor/PhysicsAssetEditor/Private/PhysicsAssetEditor.cpp
+++ b/Engine/Source/Editor/PhysicsAssetEditor/Private/PhysicsAssetEditor.cpp
@@ -449,6 +449,12 @@ void FPhysicsAssetEditor::OnAssetReimport(UObject* Object)
 
 void FPhysicsAssetEditor::OnFinishedChangingProperties(const FPropertyChangedEvent& PropertyChangedEvent)
 {
+	// if simulating ignore update request
+	if (SharedData->bRunningSimulation)
+	{
+		return;
+	}
+
 	FName PropertyName = (PropertyChangedEvent.Property != NULL) ? PropertyChangedEvent.Property->GetFName() : NAME_None;
 
 	// Update bounds bodies and setup when bConsiderForBounds was changed
