Fix white on white styling in the physics asset editor profile picker

diff --git a/Engine/Source/Editor/PhysicsAssetEditor/Private/PhysicsAssetDetailsCustomization.cpp b/Engine/Source/Editor/PhysicsAssetEditor/Private/PhysicsAssetDetailsCustomization.cpp
--- a/Engine/Source/Editor/PhysicsAssetEditor/Private/PhysicsAssetDetailsCustomization.cpp
+++ b/Engine/Source/Editor/PhysicsAssetEditor/Private/PhysicsAssetDetailsCustomization.cpp
@@ -471,9 +471,7 @@ TSharedRef<SWidget> FPhysicsAssetDetailsCustomization::MakePhysicalAnimationProf
 		.VAlign(VAlign_Center)
 		[
 			SNew(SComboButton)
-			.ButtonStyle(FEditorStyle::Get(), "ToolBar.Button")
 			.OnGetMenuContent(this, &FPhysicsAssetDetailsCustomization::FillPhysicalAnimationProfileOptions)
-			.ForegroundColor(FEditorStyle::GetSlateColor("DefaultForeground"))
 			.ButtonContent()
 			[
 				SNew(SVerticalBox)
@@ -514,7 +512,6 @@ TSharedRef<SWidget> FPhysicsAssetDetailsCustomization::MakePhysicalAnimationProf
 								PhysicalAnimationProfileNameTextBox->SetError(FText::GetEmpty());
 							}
 						})
-						.Style(FEditorStyle::Get(), "PhysicsAssetEditor.Profiles.EditableTextBoxStyle")
 						.OnTextCommitted(FOnTextCommitted::CreateSP(this, &FPhysicsAssetDetailsCustomization::HandlePhysicalAnimationProfileNameCommitted))
 					]
 				]
@@ -558,9 +555,7 @@ TSharedRef<SWidget> FPhysicsAssetDetailsCustomization::MakeConstraintProfilesWid
 		.VAlign(VAlign_Center)
 		[
 			SNew(SComboButton)
-			.ButtonStyle(FEditorStyle::Get(), "ToolBar.Button")
 			.OnGetMenuContent(this, &FPhysicsAssetDetailsCustomization::FillConstraintProfilesOptions)
-			.ForegroundColor(FEditorStyle::GetSlateColor("DefaultForeground"))
 			.ButtonContent()
 			[
 				SNew(SVerticalBox)
@@ -601,7 +596,6 @@ TSharedRef<SWidget> FPhysicsAssetDetailsCustomization::MakeConstraintProfilesWid
 								ConstraintProfileNameTextBox->SetError(FText::GetEmpty());
 							}
 						})
-						.Style(FEditorStyle::Get(), "PhysicsAssetEditor.Profiles.EditableTextBoxStyle")
 						.OnTextCommitted(FOnTextCommitted::CreateSP(this, &FPhysicsAssetDetailsCustomization::HandleConstraintProfileNameCommitted))
 					]
 				]
