Physics Asset Editor : add support for the no-gizmo selection mode

diff --git a/Engine/Source/Editor/PhysicsAssetEditor/Private/PhysicsAssetEditorEditMode.cpp b/Engine/Source/Editor/PhysicsAssetEditor/Private/PhysicsAssetEditorEditMode.cpp
--- a/Engine/Source/Editor/PhysicsAssetEditor/Private/PhysicsAssetEditorEditMode.cpp
+++ b/Engine/Source/Editor/PhysicsAssetEditor/Private/PhysicsAssetEditorEditMode.cpp
@@ -622,7 +622,7 @@ bool FPhysicsAssetEditorEditMode::UsesTransformWidget(UE::Widget::EWidgetMode Ch
 		return false;
 	}
 
-	return ShouldDrawWidget() && (CheckMode == UE::Widget::WM_Scale || CheckMode == UE::Widget::WM_Translate || CheckMode == UE::Widget::WM_Rotate);
+	return ShouldDrawWidget() && (CheckMode == UE::Widget::WM_Scale || CheckMode == UE::Widget::WM_Translate || CheckMode == UE::Widget::WM_Rotate || CheckMode == UE::Widget::WM_None);
 }
 
 bool FPhysicsAssetEditorEditMode::GetCustomDrawingCoordinateSystem(FMatrix& InMatrix, void* InData)
