Physics asset editor : improvement and bug fixes
- add isolate mode for showing colliding rigid bodies ( hotkey : C )
- fix right click menu clearing multi-selection

diff --git a/Engine/Source/Editor/Persona/Private/AnimationEditorPreviewScene.cpp b/Engine/Source/Editor/Persona/Private/AnimationEditorPreviewScene.cpp
--- a/Engine/Source/Editor/Persona/Private/AnimationEditorPreviewScene.cpp
+++ b/Engine/Source/Editor/Persona/Private/AnimationEditorPreviewScene.cpp
@@ -53,6 +53,7 @@ FAnimationEditorPreviewScene::FAnimationEditorPreviewScene(const ConstructionVal
 	, bSelecting(false)
 	, bAllowAdditionalMeshes(true)
 	, bAdditionalMeshesSelectable(true)
+	, bUsePhysicsBodiesForBoneSelection(true)
 {
 	if (GEditor)
 	{
@@ -1101,6 +1102,16 @@ bool FAnimationEditorPreviewScene::AllowMeshHitProxies() const
 	return bEnableMeshHitProxies;
 }
 
+bool FAnimationEditorPreviewScene::UsePhysicsBodiesForBoneSelection() const
+{
+	return bUsePhysicsBodiesForBoneSelection;
+}
+
+void FAnimationEditorPreviewScene::SetUsePhysicsBodiesForBoneSelection(bool bUsePhysicsBodies)
+{
+	bUsePhysicsBodiesForBoneSelection = bUsePhysicsBodies;
+}
+
 void FAnimationEditorPreviewScene::SetAllowMeshHitProxies(bool bState)
 {
 	bEnableMeshHitProxies = bState;
diff --git a/Engine/Source/Editor/Persona/Private/AnimationEditorPreviewScene.h b/Engine/Source/Editor/Persona/Private/AnimationEditorPreviewScene.h
--- a/Engine/Source/Editor/Persona/Private/AnimationEditorPreviewScene.h
+++ b/Engine/Source/Editor/Persona/Private/AnimationEditorPreviewScene.h
@@ -317,6 +317,13 @@ public:
 		bAllowAdditionalMeshes = bAllow;
 	}
 
+	/** Get whether bones can be selected by their physics bodies */
+	virtual bool UsePhysicsBodiesForBoneSelection() const override;
+
+	/** Set whether bones can be selected by their physics bodies */
+	virtual void SetUsePhysicsBodiesForBoneSelection(bool bUsePhysicsBodies) override;
+
+
 private:
 	/** Set preview mesh internal use only. The mesh should be verified by now. */
 	void SetPreviewMeshInternal(USkeletalMesh* NewPreviewMesh);
@@ -461,6 +468,9 @@ private:
 	/** Allow additional meshes to be selectable */
 	bool bAdditionalMeshesSelectable;
 
+	/** Allow bones to be selectable by clicking on their respective physics bodies */
+	bool bUsePhysicsBodiesForBoneSelection;
+
 	/** Delegate Remove attach component */
 	FOnRemoveAttachedComponentFilter OnRemoveAttachedComponentFilter;
 };
diff --git a/Engine/Source/Editor/Persona/Private/EditModes/SkeletonSelectionEditMode.cpp b/Engine/Source/Editor/Persona/Private/EditModes/SkeletonSelectionEditMode.cpp
--- a/Engine/Source/Editor/Persona/Private/EditModes/SkeletonSelectionEditMode.cpp
+++ b/Engine/Source/Editor/Persona/Private/EditModes/SkeletonSelectionEditMode.cpp
@@ -558,8 +558,9 @@ bool FSkeletonSelectionEditMode::HandleClick(FEditorViewportClient* InViewportCl
 			MeshComponent->SetSelectedEditorSection(INDEX_NONE);
 		}
 	}
-	
-	if ( !bHandled && !bSelectingSections )
+
+	const bool bUsePhysicsBodiesForBoneSelection = GetAnimPreviewScene().UsePhysicsBodiesForBoneSelection();
+	if ( !bHandled && !bSelectingSections && bUsePhysicsBodiesForBoneSelection)
 	{
 		// Cast for phys bodies if we didn't get any hit proxies
 		FHitResult Result(1.0f);
diff --git a/Engine/Source/Editor/Persona/Public/IPersonaPreviewScene.h b/Engine/Source/Editor/Persona/Public/IPersonaPreviewScene.h
--- a/Engine/Source/Editor/Persona/Public/IPersonaPreviewScene.h
+++ b/Engine/Source/Editor/Persona/Public/IPersonaPreviewScene.h
@@ -298,4 +298,10 @@ public:
 	virtual void FlagTickable() = 0;
 	/** Handle syncing selection with the skeleton tree */
 	virtual void HandleSkeletonTreeSelectionChanged(const TArrayView<TSharedPtr<ISkeletonTreeItem>>& InSelectedItems, ESelectInfo::Type InSelectInfo) = 0;
+
+	/** Get whether bones can be selected by their physics bodies */
+	virtual bool UsePhysicsBodiesForBoneSelection() const = 0;
+
+	/** Set whether bones can be selected by their physics bodies */
+	virtual void SetUsePhysicsBodiesForBoneSelection(bool bUsePhysicsBodies) = 0;
 };
diff --git a/Engine/Source/Editor/PhysicsAssetEditor/Private/PhysicsAssetEditor.cpp b/Engine/Source/Editor/PhysicsAssetEditor/Private/PhysicsAssetEditor.cpp
--- a/Engine/Source/Editor/PhysicsAssetEditor/Private/PhysicsAssetEditor.cpp
+++ b/Engine/Source/Editor/PhysicsAssetEditor/Private/PhysicsAssetEditor.cpp
@@ -651,6 +651,7 @@ void FPhysicsAssetEditor::ExtendMenu()
 			MenuBarBuilder.AddMenuEntry(Commands.ShowSelected);
 			MenuBarBuilder.AddMenuEntry(Commands.HideSelected);
 			MenuBarBuilder.AddMenuEntry(Commands.ToggleShowOnlySelected);
+			MenuBarBuilder.AddMenuEntry(Commands.ToggleShowOnlyColliding);
 			MenuBarBuilder.AddMenuEntry(Commands.ShowAll);
 			MenuBarBuilder.AddMenuEntry(Commands.HideAll);
 			MenuBarBuilder.AddMenuEntry(Commands.DeselectAll);
@@ -955,6 +956,11 @@ void FPhysicsAssetEditor::BindCommands()
 		FExecuteAction::CreateSP(this, &FPhysicsAssetEditor::OnHideSelected),
 		FCanExecuteAction::CreateSP(this, &FPhysicsAssetEditor::IsNotSimulation));
 
+	ToolkitCommands->MapAction(
+		Commands.ToggleShowOnlyColliding,
+		FExecuteAction::CreateSP(this, &FPhysicsAssetEditor::OnToggleShowOnlyColliding),
+		FCanExecuteAction::CreateSP(this, &FPhysicsAssetEditor::HasOneSelectedBodyAndIsNotSimulation));
+
 	ToolkitCommands->MapAction(
 		Commands.ToggleShowOnlySelected,
 		FExecuteAction::CreateSP(this, &FPhysicsAssetEditor::OnToggleShowOnlySelected),
@@ -1397,6 +1403,7 @@ void FPhysicsAssetEditor::BuildMenuWidgetSelection(FMenuBuilder& InMenuBuilder)
 		InMenuBuilder.AddMenuEntry( Commands.ShowSelected );
 		InMenuBuilder.AddMenuEntry( Commands.HideSelected );
 		InMenuBuilder.AddMenuEntry( Commands.ToggleShowOnlySelected );
+		InMenuBuilder.AddMenuEntry( Commands.ToggleShowOnlyColliding );
 		InMenuBuilder.AddMenuEntry( Commands.ShowAll );
 		InMenuBuilder.AddMenuEntry( Commands.HideAll );
 		InMenuBuilder.EndSection();
@@ -1793,6 +1800,11 @@ bool FPhysicsAssetEditor::HasSelectedBodyAndIsNotSimulation() const
 	return IsNotSimulation() && (SharedData->GetSelectedBody());
 }
 
+bool FPhysicsAssetEditor::HasOneSelectedBodyAndIsNotSimulation() const
+{
+	return IsNotSimulation() && (SharedData->SelectedBodies.Num() == 1);
+}
+
 bool FPhysicsAssetEditor::CanEditConstraintProperties() const
 {
 	if(IsNotSimulation() && SharedData->PhysicsAsset && SharedData->GetSelectedConstraint())
@@ -2923,6 +2935,11 @@ void FPhysicsAssetEditor::OnHideSelected()
 	SharedData->HideSelected();
 }
 
+void FPhysicsAssetEditor::OnToggleShowOnlyColliding()
+{
+	SharedData->ToggleShowOnlyColliding();
+}
+
 void FPhysicsAssetEditor::OnToggleShowOnlySelected()
 {
 	SharedData->ToggleShowOnlySelected();
@@ -3161,6 +3178,7 @@ void FPhysicsAssetEditor::HandlePreviewSceneCreated(const TSharedRef<IPersonaPre
 	InPersonaPreviewScene->SetPreviewMeshComponent(SharedData->EditorSkelComp);
 	InPersonaPreviewScene->AddComponent(SharedData->EditorSkelComp, FTransform::Identity);
 	InPersonaPreviewScene->SetAdditionalMeshesSelectable(false);
+	InPersonaPreviewScene->SetUsePhysicsBodiesForBoneSelection(false);
 	// set root component, so we can attach to it. 
 	Actor->SetRootComponent(SharedData->EditorSkelComp);
 
diff --git a/Engine/Source/Editor/PhysicsAssetEditor/Private/PhysicsAssetEditor.h b/Engine/Source/Editor/PhysicsAssetEditor/Private/PhysicsAssetEditor.h
--- a/Engine/Source/Editor/PhysicsAssetEditor/Private/PhysicsAssetEditor.h
+++ b/Engine/Source/Editor/PhysicsAssetEditor/Private/PhysicsAssetEditor.h
@@ -201,6 +201,7 @@ private:
 
 	/** Toolbar/menu command methods */
 	bool HasSelectedBodyAndIsNotSimulation() const;
+	bool HasOneSelectedBodyAndIsNotSimulation() const;
 	bool CanEditConstraintProperties() const;
 	bool HasSelectedConstraintAndIsNotSimulation() const;
 	void OnChangeDefaultMesh(USkeletalMesh* OldPreviewMesh, USkeletalMesh* NewPreviewMesh);
@@ -290,6 +291,7 @@ private:
 	void OnToggleShowSelected();
 	void OnShowSelected();
 	void OnHideSelected();
+	void OnToggleShowOnlyColliding();
 	void OnToggleShowOnlySelected();
 	void OnShowAll();
 	void OnHideAll();
diff --git a/Engine/Source/Editor/PhysicsAssetEditor/Private/PhysicsAssetEditorActions.cpp b/Engine/Source/Editor/PhysicsAssetEditor/Private/PhysicsAssetEditorActions.cpp
--- a/Engine/Source/Editor/PhysicsAssetEditor/Private/PhysicsAssetEditorActions.cpp
+++ b/Engine/Source/Editor/PhysicsAssetEditor/Private/PhysicsAssetEditorActions.cpp
@@ -83,6 +83,7 @@ void FPhysicsAssetEditorCommands::RegisterCommands()
 	UI_COMMAND(ToggleShowSelected, "Toggle Show Selected", "Show/hide selected bodies/constraints.", EUserInterfaceActionType::Button, FInputChord(EKeys::H));
 	UI_COMMAND(ShowSelected, "Show Selected", "Show selected bodies/constraints.", EUserInterfaceActionType::Button, FInputChord(EModifierKey::Shift, EKeys::H));
 	UI_COMMAND(HideSelected, "Hide Selected", "Hide selected bodies/constraints.", EUserInterfaceActionType::Button, FInputChord(EModifierKey::Control, EKeys::H));
+	UI_COMMAND(ToggleShowOnlyColliding, "Toggle Isolate Colliding Bodies", "Show only the selected body and the ones that can collide with it, hiding all others or Show all.", EUserInterfaceActionType::Button, FInputChord(EKeys::C));
 	UI_COMMAND(ToggleShowOnlySelected, "Toggle Isolate Selected", "Show the selected bodies/constraints, hiding all others, or Show all.", EUserInterfaceActionType::Button, FInputChord(EKeys::G));
 	UI_COMMAND(ShowAll, "Show All", "Show all bodies/constraints.", EUserInterfaceActionType::Button, FInputChord(EModifierKey::Shift, EKeys::G));
 	UI_COMMAND(HideAll, "Hide All", "Hide all bodies/constraints.", EUserInterfaceActionType::Button, FInputChord(EModifierKey::Control, EKeys::G));
diff --git a/Engine/Source/Editor/PhysicsAssetEditor/Private/PhysicsAssetEditorActions.h b/Engine/Source/Editor/PhysicsAssetEditor/Private/PhysicsAssetEditorActions.h
--- a/Engine/Source/Editor/PhysicsAssetEditor/Private/PhysicsAssetEditorActions.h
+++ b/Engine/Source/Editor/PhysicsAssetEditor/Private/PhysicsAssetEditorActions.h
@@ -95,6 +95,7 @@ public:
 	TSharedPtr<FUICommandInfo> ToggleShowSelected;
 	TSharedPtr<FUICommandInfo> ShowSelected;
 	TSharedPtr<FUICommandInfo> HideSelected;
+	TSharedPtr<FUICommandInfo> ToggleShowOnlyColliding;
 	TSharedPtr<FUICommandInfo> ToggleShowOnlySelected;
 	TSharedPtr<FUICommandInfo> ShowAll;
 	TSharedPtr<FUICommandInfo> HideAll;
diff --git a/Engine/Source/Editor/PhysicsAssetEditor/Private/PhysicsAssetEditorSharedData.cpp b/Engine/Source/Editor/PhysicsAssetEditor/Private/PhysicsAssetEditorSharedData.cpp
--- a/Engine/Source/Editor/PhysicsAssetEditor/Private/PhysicsAssetEditorSharedData.cpp
+++ b/Engine/Source/Editor/PhysicsAssetEditor/Private/PhysicsAssetEditorSharedData.cpp
@@ -742,6 +742,27 @@ void FPhysicsAssetEditorSharedData::HideSelected()
 	}
 }
 
+void FPhysicsAssetEditorSharedData::ToggleShowOnlyColliding()
+{
+	// important that we check this before calling ShowAll
+	const bool bIsShowingColliding = (HiddenBodies == NoCollisionBodies);
+
+	// in any case first show all
+	ShowAll();
+
+	if (!bIsShowingColliding)
+	{
+		// only works if one only body is selected
+		if (PhysicsAsset != nullptr && SelectedBodies.Num() == 1)
+		{
+
+			// NoCollisionBodies already contains the non colliding bodies from the one selection
+			HiddenBodies.Empty();
+			HiddenBodies.Append(NoCollisionBodies);
+		}
+	}
+}
+
 void FPhysicsAssetEditorSharedData::UpdateNoCollisionBodies()
 {
 	NoCollisionBodies.Empty();
diff --git a/Engine/Source/Editor/PhysicsAssetEditor/Private/PhysicsAssetEditorSharedData.h b/Engine/Source/Editor/PhysicsAssetEditor/Private/PhysicsAssetEditorSharedData.h
--- a/Engine/Source/Editor/PhysicsAssetEditor/Private/PhysicsAssetEditorSharedData.h
+++ b/Engine/Source/Editor/PhysicsAssetEditor/Private/PhysicsAssetEditorSharedData.h
@@ -145,6 +145,7 @@ public:
 	void ToggleShowSelected();
 	void ShowAll();
 	void HideAll();
+	void ToggleShowOnlyColliding();
 	void ToggleShowOnlySelected();
 	void ShowSelected();
 	void HideSelected();
