Physics Asset Editor - Fix selection synchronization when changing filters

diff --git a/Engine/Source/Editor/PhysicsAssetEditor/Private/PhysicsAssetEditor.cpp b/Engine/Source/Editor/PhysicsAssetEditor/Private/PhysicsAssetEditor.cpp
--- a/Engine/Source/Editor/PhysicsAssetEditor/Private/PhysicsAssetEditor.cpp
+++ b/Engine/Source/Editor/PhysicsAssetEditor/Private/PhysicsAssetEditor.cpp
@@ -3261,31 +3261,31 @@ void FPhysicsAssetEditor::HandleExtendFilterMenu(FMenuBuilder& InMenuBuilder)
 void FPhysicsAssetEditor::HandleToggleShowBodies()
 {
 	SkeletonTreeBuilder->bShowBodies = !SkeletonTreeBuilder->bShowBodies;
-	SkeletonTree->RefreshFilter();
+	RefreshFilter();
 }
 
 void FPhysicsAssetEditor::HandleToggleShowSimulatedBodies()
 {
 	SkeletonTreeBuilder->bShowSimulatedBodies = !SkeletonTreeBuilder->bShowSimulatedBodies;
-	SkeletonTree->RefreshFilter();
+	RefreshFilter();
 }
 
 void FPhysicsAssetEditor::HandleToggleShowKinematicBodies()
 {
 	SkeletonTreeBuilder->bShowKinematicBodies = !SkeletonTreeBuilder->bShowKinematicBodies;
-	SkeletonTree->RefreshFilter();
+	RefreshFilter();
 }
 
 void FPhysicsAssetEditor::HandleToggleShowConstraints()
 {
 	SkeletonTreeBuilder->bShowConstraints = !SkeletonTreeBuilder->bShowConstraints;
-	SkeletonTree->RefreshFilter();
+	RefreshFilter();
 }
 
 void FPhysicsAssetEditor::HandleToggleShowPrimitives()
 {
 	SkeletonTreeBuilder->bShowPrimitives = !SkeletonTreeBuilder->bShowPrimitives;
-	SkeletonTree->RefreshFilter();
+	RefreshFilter();
 }
 
 ECheckBoxState FPhysicsAssetEditor::GetShowBodiesChecked() const
@@ -3331,6 +3331,13 @@ void FPhysicsAssetEditor::HandleGetFilterLabel(TArray<FText>& InOutItems) const
 	}
 }
 
+void FPhysicsAssetEditor::RefreshFilter()
+{
+	SkeletonTree->RefreshFilter();
+	// make sure we resynchronize the list 
+	HandleViewportSelectionChanged(SharedData->SelectedBodies, SharedData->SelectedConstraints);
+}
+
 void FPhysicsAssetEditor::HandleCreateNewConstraint(int32 BodyIndex0, int32 BodyIndex1)
 {
 	if(BodyIndex0 != BodyIndex1)
diff --git a/Engine/Source/Editor/PhysicsAssetEditor/Private/PhysicsAssetEditor.h b/Engine/Source/Editor/PhysicsAssetEditor/Private/PhysicsAssetEditor.h
--- a/Engine/Source/Editor/PhysicsAssetEditor/Private/PhysicsAssetEditor.h
+++ b/Engine/Source/Editor/PhysicsAssetEditor/Private/PhysicsAssetEditor.h
@@ -324,6 +324,9 @@ private:
 	/** Customize the filter label */
 	void HandleGetFilterLabel(TArray<FText>& InOutItems) const;
 
+	/** refresh filter after changing filter settings */
+	void RefreshFilter();
+
 	/** Invalidate convex meshes and recreate the physics state. Performed on property changes (etc) */
 	void RecreatePhysicsState();
 
