UE-110714 - SGraphPanel::OnPaint crash in physics asset editor

Upon construction, SPhysicsAssetGraphNodeOutputPin objects will register callbacks for its pins. However, these objects can potentially reference destroyed UEdGraphPin instances after an undo occurs, since UEdGraphPin deserializaiton will destroy existing pins. There's also a secondary issue where the nodes appear to be disconnected after undo.

The most straightforward fix here is to make SPhysicsAssetGraph a FEditorUndoClient, and then rebuild the underlying graph object.

diff --git a/Engine/Source/Editor/PhysicsAssetEditor/Private/PhysicsAssetGraph/SPhysicsAssetGraph.cpp b/Engine/Source/Editor/PhysicsAssetEditor/Private/PhysicsAssetGraph/SPhysicsAssetGraph.cpp
--- a/Engine/Source/Editor/PhysicsAssetEditor/Private/PhysicsAssetGraph/SPhysicsAssetGraph.cpp
+++ b/Engine/Source/Editor/PhysicsAssetEditor/Private/PhysicsAssetGraph/SPhysicsAssetGraph.cpp
@@ -82,6 +82,8 @@ void SPhysicsAssetGraph::Construct(const FArguments& InArgs, const TSharedRef<FP
 	];
 
 	BlockSelectionBroadcastCounter = 2;
+
+	GEditor->RegisterForUndo(this);
 }
 
 FActionMenuContent SPhysicsAssetGraph::OnCreateGraphActionMenu(UEdGraph* InGraph, const FVector2D& InNodePosition, const TArray<UEdGraphPin*>& InDraggedPins, bool bAutoExpand, SGraphEditor::FActionMenuClosed InOnMenuClosed)
@@ -180,4 +182,14 @@ void SPhysicsAssetGraph::Tick(const FGeometry& AllottedGeometry, const double In
 	BlockSelectionBroadcastCounter--;
 }
 
+void SPhysicsAssetGraph::PostUndo(bool bSuccess)
+{
+	GraphObj->RebuildGraph();
+}
+
+void SPhysicsAssetGraph::PostRedo(bool bSuccess)
+{
+	GraphObj->RebuildGraph();
+}
+
 #undef LOCTEXT_NAMESPACE
diff --git a/Engine/Source/Editor/PhysicsAssetEditor/Private/PhysicsAssetGraph/SPhysicsAssetGraph.h b/Engine/Source/Editor/PhysicsAssetEditor/Private/PhysicsAssetGraph/SPhysicsAssetGraph.h
--- a/Engine/Source/Editor/PhysicsAssetEditor/Private/PhysicsAssetGraph/SPhysicsAssetGraph.h
+++ b/Engine/Source/Editor/PhysicsAssetEditor/Private/PhysicsAssetGraph/SPhysicsAssetGraph.h
@@ -5,6 +5,7 @@
 #include "PersonaDelegates.h"
 #include "GraphEditor.h"
 #include "Containers/ArrayView.h"
+#include "EditorUndoClient.h"
 
 class UObject;
 class UPhysicsAsset;
@@ -22,7 +23,7 @@ DECLARE_DELEGATE_OneParam(FOnPhysicsAssetGraphCreated, const TSharedRef<SPhysics
 /** Delegate used to communicate graph selection */
 DECLARE_DELEGATE_OneParam(FOnGraphObjectsSelected, const TArrayView<UObject*>& /*InObjects*/);
 
-class SPhysicsAssetGraph : public SCompoundWidget
+class SPhysicsAssetGraph : public SCompoundWidget, public FEditorUndoClient
 {
 public:
 	SLATE_BEGIN_ARGS(SPhysicsAssetGraph){}
@@ -34,6 +35,10 @@ public:
 	/** SWidget interface */
 	virtual void Tick(const FGeometry& AllottedGeometry, const double InCurrentTime, const float InDeltaTime) override;
 
+	/** FEditorUndoClient interface */
+	virtual void PostUndo(bool bSuccess) override;
+	virtual void PostRedo(bool bSuccess) override;
+
 	/** Constructs this widget with InArgs */
 	void Construct(const FArguments& InArgs, const TSharedRef<FPhysicsAssetEditor>& InPhysicsAssetEditor, UPhysicsAsset* InPhysicsAsset, const TSharedRef<IEditableSkeleton>& InEditableSkeleton, FOnGraphObjectsSelected InOnGraphObjectSelected);
 
