Improved copy/paste properties behavior in the physics asset editor to allow for cross-asset copy/paste.
Fixed issue where the rotation would be copied as well when it shouldn't be.

diff --git a/Engine/Source/Editor/PhysicsAssetEditor/PhysicsAssetEditor.Build.cs b/Engine/Source/Editor/PhysicsAssetEditor/PhysicsAssetEditor.Build.cs
--- a/Engine/Source/Editor/PhysicsAssetEditor/PhysicsAssetEditor.Build.cs
+++ b/Engine/Source/Editor/PhysicsAssetEditor/PhysicsAssetEditor.Build.cs
@@ -46,7 +46,8 @@ public class PhysicsAssetEditor : ModuleRules
 				"PhysicsCore",
 				"PhysicsUtilities",
 				"MeshUtilitiesCommon",
-            }
+				"ApplicationCore",
+			}
         );
 
 		DynamicallyLoadedModuleNames.AddRange(
diff --git a/Engine/Source/Editor/PhysicsAssetEditor/Private/PhysicsAssetEditor.cpp b/Engine/Source/Editor/PhysicsAssetEditor/Private/PhysicsAssetEditor.cpp
--- a/Engine/Source/Editor/PhysicsAssetEditor/Private/PhysicsAssetEditor.cpp
+++ b/Engine/Source/Editor/PhysicsAssetEditor/Private/PhysicsAssetEditor.cpp
@@ -2068,7 +2068,7 @@ bool FPhysicsAssetEditor::CanPasteProperties() const
 
 bool FPhysicsAssetEditor::IsCopyProperties() const
 {
-	return (SharedData->CopiedBodySetup) || (SharedData->CopiedConstraintTemplate);
+	return FPhysicsAssetEditorSharedData::ClipboardHasCompatibleData();
 }
 
 //We need to save and restore physics states based on the mode we use to simulate
diff --git a/Engine/Source/Editor/PhysicsAssetEditor/Private/PhysicsAssetEditorSharedData.cpp b/Engine/Source/Editor/PhysicsAssetEditor/Private/PhysicsAssetEditorSharedData.cpp
--- a/Engine/Source/Editor/PhysicsAssetEditor/Private/PhysicsAssetEditorSharedData.cpp
+++ b/Engine/Source/Editor/PhysicsAssetEditor/Private/PhysicsAssetEditorSharedData.cpp
@@ -34,11 +34,18 @@
 #include "Widgets/Input/SButton.h"
 #include "Widgets/Text/STextBlock.h"
 #include "ClothingSimulationInteractor.h"
+#include "HAL/PlatformApplicationMisc.h"
 
 #define LOCTEXT_NAMESPACE "PhysicsAssetEditorShared"
 
 //PRAGMA_DISABLE_OPTIMIZATION
 
+namespace SharedDataConstants
+{
+	const FString ConstraintType = TEXT("Constraint");
+	const FString BodyType = TEXT("SkeletalBodySetup");
+}
+
 
 FScopedBulkSelection::FScopedBulkSelection(TSharedPtr<FPhysicsAssetEditorSharedData> InSharedData)
 	: SharedData(InSharedData)
@@ -54,8 +61,6 @@ FScopedBulkSelection::~FScopedBulkSelection()
 
 FPhysicsAssetEditorSharedData::FPhysicsAssetEditorSharedData()
 	: COMRenderColor(255,255,100)
-	, CopiedBodySetup(NULL)
-	, CopiedConstraintTemplate(NULL)
 	, bSuspendSelectionBroadcast(false)
 	, InsideSelChange(0)
 {
@@ -231,14 +236,14 @@ void FPhysicsAssetEditorSharedData::CachePreviewMesh()
 	}
 }
 
-void FPhysicsAssetEditorSharedData::CopyConstraintProperties(UPhysicsConstraintTemplate * FromConstraintSetup, UPhysicsConstraintTemplate * ToConstraintSetup)
+void FPhysicsAssetEditorSharedData::CopyConstraintProperties(const UPhysicsConstraintTemplate * FromConstraintSetup, UPhysicsConstraintTemplate * ToConstraintSetup, bool bKeepOldRotation)
 {
 	ToConstraintSetup->Modify();
 	FConstraintInstance OldInstance = ToConstraintSetup->DefaultInstance;
 	ToConstraintSetup->DefaultInstance.CopyConstraintParamsFrom(&FromConstraintSetup->DefaultInstance);
 
 	// recover certain data that we'd like to keep - i.e. bone indices those still should stay.  
-	// frame position offsets taken from old, but frame orientations are taken from new source
+	// frame position offsets taken from old, but, optionally, frame orientations are taken from new source
 	ToConstraintSetup->DefaultInstance.ConstraintIndex = OldInstance.ConstraintIndex;
 #if WITH_PHYSX
 	ToConstraintSetup->DefaultInstance.ConstraintHandle = OldInstance.ConstraintHandle;
@@ -249,9 +254,86 @@ void FPhysicsAssetEditorSharedData::CopyConstraintProperties(UPhysicsConstraintT
 	ToConstraintSetup->DefaultInstance.Pos1 = OldInstance.Pos1;
 	ToConstraintSetup->DefaultInstance.Pos2 = OldInstance.Pos2;
 
+	if (bKeepOldRotation)
+	{
+		ToConstraintSetup->DefaultInstance.PriAxis1 = OldInstance.PriAxis1;
+		ToConstraintSetup->DefaultInstance.SecAxis1 = OldInstance.SecAxis1;
+		ToConstraintSetup->DefaultInstance.PriAxis2 = OldInstance.PriAxis2;
+		ToConstraintSetup->DefaultInstance.SecAxis2 = OldInstance.SecAxis2;
+		ToConstraintSetup->DefaultInstance.AngularRotationOffset = OldInstance.AngularRotationOffset;
+	}
+
 	ToConstraintSetup->UpdateProfileInstance();
 }
 
+void FPhysicsAssetEditorSharedData::CopyToClipboard(const FString& ObjectType, UObject* Object)
+{
+	FSoftObjectPath PhysicsAssetPath(PhysicsAsset);
+	FSoftObjectPath ObjectAssetPath(Object);
+	FString ClipboardContent = FString::Format(TEXT("{0};{1};{2}"), { PhysicsAssetPath.ToString(), *ObjectType, ObjectAssetPath.ToString() });
+	FPlatformApplicationMisc::ClipboardCopy(*ClipboardContent);
+}
+
+bool FPhysicsAssetEditorSharedData::PasteFromClipboard(const FString& InObjectType, UPhysicsAsset*& OutAsset, UObject*& OutObject)
+{
+	FString SourceObjectType;
+	return ParseClipboard(OutAsset, SourceObjectType, OutObject) && SourceObjectType == InObjectType;
+}
+
+void FPhysicsAssetEditorSharedData::ConditionalClearClipboard(const FString& ObjectType, UObject* Object)
+{
+	UPhysicsAsset* SourceAsset = nullptr;
+	FString SourceObjectType;
+	UObject* SourceObject = nullptr;
+
+	if(ParseClipboard(SourceAsset, SourceObjectType, SourceObject))
+	{
+		// Clear the clipboard if it matches the parameters we're given
+		if (SourceAsset == PhysicsAsset && SourceObjectType == ObjectType && SourceObject == Object)
+		{
+			FString EmptyString;
+			FPlatformApplicationMisc::ClipboardCopy(*EmptyString);
+		}
+	}
+}
+
+bool FPhysicsAssetEditorSharedData::ClipboardHasCompatibleData()
+{
+	UPhysicsAsset* DummyAsset = nullptr;
+	FString DummyObjectType;
+	UObject* DummyObject = nullptr;
+	return ParseClipboard(DummyAsset, DummyObjectType, DummyObject);
+}
+
+bool FPhysicsAssetEditorSharedData::ParseClipboard(UPhysicsAsset*& OutAsset, FString& OutObjectType, UObject*& OutObject)
+{
+	FString ClipboardContent;
+	FPlatformApplicationMisc::ClipboardPaste(ClipboardContent);
+
+	TArray<FString> ParsedString;
+	ClipboardContent.ParseIntoArray(ParsedString, TEXT(";"), true);
+
+	if (ParsedString.Num() != 3)
+	{
+		return false;
+	}
+
+	FSoftObjectPath PhysicsAssetPath(ParsedString[0]);
+	OutAsset = Cast<UPhysicsAsset>(PhysicsAssetPath.ResolveObject());
+
+	if (!OutAsset)
+	{
+		return false;
+	}
+
+	OutObjectType = ParsedString[1];
+
+	FSoftObjectPath ObjectAssetPath(ParsedString[2]);
+	OutObject = ObjectAssetPath.ResolveObject();
+
+	return OutObject != nullptr;
+}
+
 struct FMirrorInfo
 {
 	FName BoneName;
@@ -1360,8 +1442,7 @@ void FPhysicsAssetEditorSharedData::AutoNamePrimitive(int32 BodyIndex, EAggColli
 void FPhysicsAssetEditorSharedData::CopyBody()
 {
 	check(SelectedBodies.Num() == 1);
-
-	CopiedBodySetup = PhysicsAsset->SkeletalBodySetups[GetSelectedBody()->Index];
+	CopyToClipboard(SharedDataConstants::BodyType, PhysicsAsset->SkeletalBodySetups[GetSelectedBody()->Index]);
 }
 
 void FPhysicsAssetEditorSharedData::PasteBodyProperties()
@@ -1372,6 +1453,17 @@ void FPhysicsAssetEditorSharedData::PasteBodyProperties()
 		return;
 	}
 
+	UPhysicsAsset* SourceAsset = nullptr;
+	UObject* SourceBodySetup = nullptr;
+	int32 SourceBodyIndex = 0;
+
+	if(!PasteFromClipboard(SharedDataConstants::BodyType, SourceAsset, SourceBodySetup))
+	{
+		return;
+	}
+
+	const UBodySetup* CopiedBodySetup = Cast<UBodySetup>(SourceBodySetup);
+
 	// Must have two valid bodies (which are different)
 	if(CopiedBodySetup == NULL)
 	{
@@ -1387,9 +1479,8 @@ void FPhysicsAssetEditorSharedData::PasteBodyProperties()
 		for(int32 i=0; i<SelectedBodies.Num(); ++i)
 		{
 			UBodySetup* ToBodySetup = PhysicsAsset->SkeletalBodySetups[SelectedBodies[i].Index];
-			UBodySetup* FromBodySetup = CopiedBodySetup;
 			ToBodySetup->Modify();
-			ToBodySetup->CopyBodyPropertiesFrom(FromBodySetup);
+			ToBodySetup->CopyBodyPropertiesFrom(CopiedBodySetup);
 		}
 	
 		ClearSelectedBody();	//paste can change the primitives on our selected bodies. There's probably a way to properly update this, but for now just deselect
@@ -1808,27 +1899,29 @@ void FPhysicsAssetEditorSharedData::SnapConstraintToBone(FConstraintInstance& Co
 void FPhysicsAssetEditorSharedData::CopyConstraint()
 {
 	check(SelectedConstraints.Num() == 1);
-
-	CopiedConstraintTemplate = PhysicsAsset->ConstraintSetup[GetSelectedConstraint()->Index];
+	CopyToClipboard(SharedDataConstants::ConstraintType, PhysicsAsset->ConstraintSetup[GetSelectedConstraint()->Index]);
 }
 
 void FPhysicsAssetEditorSharedData::PasteConstraintProperties()
 {
-	if (CopiedConstraintTemplate == NULL)
+	UPhysicsAsset* SourceAsset = nullptr;
+	UObject* SourceConstraint;
+
+	if(!PasteFromClipboard(SharedDataConstants::ConstraintType, SourceAsset, SourceConstraint))
 	{
 		return;
 	}
 
-	if(SelectedConstraints.Num() > 0)
-	{
-		const FScopedTransaction Transaction( NSLOCTEXT("PhysicsAssetEditor", "PasteConstraintProperties", "Paste Constraint Properties") );
+	const UPhysicsConstraintTemplate* FromConstraintSetup = Cast<UPhysicsConstraintTemplate>(SourceConstraint);
 
-		UPhysicsConstraintTemplate* FromConstraintSetup = CopiedConstraintTemplate;
+	if(FromConstraintSetup && SelectedConstraints.Num() > 0)
+	{
+		const FScopedTransaction Transaction(NSLOCTEXT("PhysicsAssetEditor", "PasteConstraintProperties", "Paste Constraint Properties"));
 
 		for(int32 i=0; i<SelectedConstraints.Num(); ++i)
 		{
 			UPhysicsConstraintTemplate* ToConstraintSetup = PhysicsAsset->ConstraintSetup[SelectedConstraints[i].Index];
-			CopyConstraintProperties(FromConstraintSetup, ToConstraintSetup);
+			CopyConstraintProperties(FromConstraintSetup, ToConstraintSetup, /*bKeepOriginalRotation=*/true);
 		}
 	}
 }
@@ -2003,6 +2096,9 @@ void FPhysicsAssetEditorSharedData::DeleteBody(int32 DelBodyIndex, bool bRefresh
 		}
 	}
 
+	// Clear clipboard if it was pointing to this body
+	ConditionalClearClipboard(SharedDataConstants::BodyType, BodySetup);
+
 	// Now actually destroy body. This will destroy any constraints associated with the body as well.
 	FPhysicsAssetUtils::DestroyBody(PhysicsAsset, DelBodyIndex);
 
@@ -2093,11 +2189,6 @@ void FPhysicsAssetEditorSharedData::DeleteCurrentPrim()
 				{
 					DeleteBody(BodyIndex, false);
 				}
-
-				if (CopiedBodySetup == BodySetup)
-				{
-					CopiedBodySetup = NULL;
-				}
 			}
 		}
 	}
@@ -2220,6 +2311,7 @@ void FPhysicsAssetEditorSharedData::DeleteCurrentConstraint()
 	TArray<int32> Indices;
 	for(int32 i=0; i<SelectedConstraints.Num(); ++i)
 	{
+		ConditionalClearClipboard(SharedDataConstants::ConstraintType, PhysicsAsset->ConstraintSetup[SelectedConstraints[i].Index]);
 		Indices.Add(SelectedConstraints[i].Index);
 	}
 
@@ -2228,15 +2320,8 @@ void FPhysicsAssetEditorSharedData::DeleteCurrentConstraint()
 	//These are indices into an array, we must remove it from greatest to smallest so that the indices don't shift
 	for(int32 i=Indices.Num() - 1; i>= 0; --i)
 	{
-		
-		if(PhysicsAsset->ConstraintSetup[Indices[i]] == CopiedConstraintTemplate)
-		{
-			CopiedConstraintTemplate = NULL;
-		}
-
 		PhysicsAsset->Modify();
 		FPhysicsAssetUtils::DestroyConstraint(PhysicsAsset, Indices[i]);
-		
 	}
 	
 	ClearSelectedConstraints();
@@ -2519,8 +2604,6 @@ void FPhysicsAssetEditorSharedData::AddReferencedObjects(FReferenceCollector& Co
 	Collector.AddReferencedObject(PhysicalAnimationComponent);
 	Collector.AddReferencedObject(EditorOptions);
 	Collector.AddReferencedObject(MouseHandle);
-	Collector.AddReferencedObject(CopiedBodySetup);
-	Collector.AddReferencedObject(CopiedConstraintTemplate);
 
 	if (PreviewScene != nullptr)
 	{
diff --git a/Engine/Source/Editor/PhysicsAssetEditor/Private/PhysicsAssetEditorSharedData.h b/Engine/Source/Editor/PhysicsAssetEditor/Private/PhysicsAssetEditorSharedData.h
--- a/Engine/Source/Editor/PhysicsAssetEditor/Private/PhysicsAssetEditorSharedData.h
+++ b/Engine/Source/Editor/PhysicsAssetEditor/Private/PhysicsAssetEditorSharedData.h
@@ -231,6 +231,9 @@ public:
 	/** broadcast a change in the preview*/
 	void BroadcastPreviewChanged();
 
+	/** Returns true if the clipboard contains data this class can process */
+	static bool ClipboardHasCompatibleData();
+
 private:
 	/** Initializes a constraint setup */
 	void InitConstraintSetup(UPhysicsConstraintTemplate* ConstraintSetup, int32 ChildBodyIndex, int32 ParentBodyIndex);
@@ -242,7 +245,19 @@ private:
 	void UpdateNoCollisionBodies();
 
 	/** Copy the properties of the one and only selected constraint */
-	void CopyConstraintProperties(UPhysicsConstraintTemplate * FromConstraintSetup, UPhysicsConstraintTemplate * ToConstraintSetup);
+	void CopyConstraintProperties(const UPhysicsConstraintTemplate * FromConstraintSetup, UPhysicsConstraintTemplate * ToConstraintSetup, bool bKeepOldRotation = false);
+
+	/** Copies a reference to a given element to the clipboard */
+	void CopyToClipboard(const FString& ObjectType, UObject* Object);
+
+	/** Pastes data from the clipboard on a given type */
+	bool PasteFromClipboard(const FString& InObjectType, UPhysicsAsset*& OutAsset, UObject*& OutObject);
+
+	/** Clears data in clipboard if it was pointing to the given type/data */
+	void ConditionalClearClipboard(const FString& ObjectType, UObject* Object);
+
+	/** Checks and parses clipboard data */
+	static bool ParseClipboard(UPhysicsAsset*& OutAsset, FString& OutObjectType, UObject*& OutObject);
 
 public:
 	/** Callback for handling selection changes */
@@ -303,9 +318,6 @@ public:
 		return Count ? &SelectedBodies[Count - 1] : NULL;
 	}
 
-	UBodySetup * CopiedBodySetup;
-	UPhysicsConstraintTemplate * CopiedConstraintTemplate;
-
 	/** Constraint editing */
 	TArray<FSelection> SelectedConstraints;
 	FSelection * GetSelectedConstraint()
