Fix/improve transparency sorting and shading in editor when opacity is 1
Option for the grab/poke raycast distance

diff --git a/Engine/Config/BaseEditorPerProjectUserSettings.ini b/Engine/Config/BaseEditorPerProjectUserSettings.ini
--- a/Engine/Config/BaseEditorPerProjectUserSettings.ini
+++ b/Engine/Config/BaseEditorPerProjectUserSettings.ini
@@ -509,6 +509,7 @@ FloorGap=25.0
 GravScale=1.0
 bPromptOnBoneDelete=true
 PokeStrength=100.0
+InteractionDistance=5000.0
 bShowNamesInHierarchy=true
 PokePauseTime=0.5
 PokeBlendTime=0.5
diff --git a/Engine/Source/Editor/PhysicsAssetEditor/Private/PhysicsAssetEditor.cpp b/Engine/Source/Editor/PhysicsAssetEditor/Private/PhysicsAssetEditor.cpp
--- a/Engine/Source/Editor/PhysicsAssetEditor/Private/PhysicsAssetEditor.cpp
+++ b/Engine/Source/Editor/PhysicsAssetEditor/Private/PhysicsAssetEditor.cpp
@@ -1021,21 +1021,21 @@ void FPhysicsAssetEditor::BindCommands()
 
 	ViewportCommandList->MapAction(
 		Commands.MeshRenderingMode_Solid,
-		FExecuteAction::CreateSP(this, &FPhysicsAssetEditor::OnMeshRenderingMode, EPhysicsAssetEditorRenderMode::Solid, false),
+		FExecuteAction::CreateSP(this, &FPhysicsAssetEditor::OnMeshRenderingMode, EPhysicsAssetEditorMeshViewMode::Solid, false),
 		FCanExecuteAction(),
-		FIsActionChecked::CreateSP(this, &FPhysicsAssetEditor::IsMeshRenderingMode, EPhysicsAssetEditorRenderMode::Solid, false));
+		FIsActionChecked::CreateSP(this, &FPhysicsAssetEditor::IsMeshRenderingMode, EPhysicsAssetEditorMeshViewMode::Solid, false));
 
 	ViewportCommandList->MapAction(
 		Commands.MeshRenderingMode_Wireframe,
-		FExecuteAction::CreateSP(this, &FPhysicsAssetEditor::OnMeshRenderingMode, EPhysicsAssetEditorRenderMode::Wireframe, false),
+		FExecuteAction::CreateSP(this, &FPhysicsAssetEditor::OnMeshRenderingMode, EPhysicsAssetEditorMeshViewMode::Wireframe, false),
 		FCanExecuteAction(),
-		FIsActionChecked::CreateSP(this, &FPhysicsAssetEditor::IsMeshRenderingMode, EPhysicsAssetEditorRenderMode::Wireframe, false));
+		FIsActionChecked::CreateSP(this, &FPhysicsAssetEditor::IsMeshRenderingMode, EPhysicsAssetEditorMeshViewMode::Wireframe, false));
 
 	ViewportCommandList->MapAction(
 		Commands.MeshRenderingMode_None,
-		FExecuteAction::CreateSP(this, &FPhysicsAssetEditor::OnMeshRenderingMode, EPhysicsAssetEditorRenderMode::None, false),
+		FExecuteAction::CreateSP(this, &FPhysicsAssetEditor::OnMeshRenderingMode, EPhysicsAssetEditorMeshViewMode::None, false),
 		FCanExecuteAction(),
-		FIsActionChecked::CreateSP(this, &FPhysicsAssetEditor::IsMeshRenderingMode, EPhysicsAssetEditorRenderMode::None, false));
+		FIsActionChecked::CreateSP(this, &FPhysicsAssetEditor::IsMeshRenderingMode, EPhysicsAssetEditorMeshViewMode::None, false));
 
 	ViewportCommandList->EndGroup();
 
@@ -1043,21 +1043,27 @@ void FPhysicsAssetEditor::BindCommands()
 
 	ViewportCommandList->MapAction(
 		Commands.CollisionRenderingMode_Solid,
-		FExecuteAction::CreateSP(this, &FPhysicsAssetEditor::OnCollisionRenderingMode, EPhysicsAssetEditorRenderMode::Solid, false),
+		FExecuteAction::CreateSP(this, &FPhysicsAssetEditor::OnCollisionRenderingMode, EPhysicsAssetEditorCollisionViewMode::Solid, false),
 		FCanExecuteAction(),
-		FIsActionChecked::CreateSP(this, &FPhysicsAssetEditor::IsCollisionRenderingMode, EPhysicsAssetEditorRenderMode::Solid, false));
+		FIsActionChecked::CreateSP(this, &FPhysicsAssetEditor::IsCollisionRenderingMode, EPhysicsAssetEditorCollisionViewMode::Solid, false));
 
 	ViewportCommandList->MapAction(
 		Commands.CollisionRenderingMode_Wireframe,
-		FExecuteAction::CreateSP(this, &FPhysicsAssetEditor::OnCollisionRenderingMode, EPhysicsAssetEditorRenderMode::Wireframe, false),
+		FExecuteAction::CreateSP(this, &FPhysicsAssetEditor::OnCollisionRenderingMode, EPhysicsAssetEditorCollisionViewMode::Wireframe, false),
 		FCanExecuteAction(),
-		FIsActionChecked::CreateSP(this, &FPhysicsAssetEditor::IsCollisionRenderingMode, EPhysicsAssetEditorRenderMode::Wireframe, false));
+		FIsActionChecked::CreateSP(this, &FPhysicsAssetEditor::IsCollisionRenderingMode, EPhysicsAssetEditorCollisionViewMode::Wireframe, false));
+
+	ViewportCommandList->MapAction(
+		Commands.CollisionRenderingMode_SolidWireframe,
+		FExecuteAction::CreateSP(this, &FPhysicsAssetEditor::OnCollisionRenderingMode, EPhysicsAssetEditorCollisionViewMode::SolidWireframe, false),
+		FCanExecuteAction(),
+		FIsActionChecked::CreateSP(this, &FPhysicsAssetEditor::IsCollisionRenderingMode, EPhysicsAssetEditorCollisionViewMode::SolidWireframe, false));
 
 	ViewportCommandList->MapAction(
 		Commands.CollisionRenderingMode_None,
-		FExecuteAction::CreateSP(this, &FPhysicsAssetEditor::OnCollisionRenderingMode, EPhysicsAssetEditorRenderMode::None, false),
+		FExecuteAction::CreateSP(this, &FPhysicsAssetEditor::OnCollisionRenderingMode, EPhysicsAssetEditorCollisionViewMode::None, false),
 		FCanExecuteAction(),
-		FIsActionChecked::CreateSP(this, &FPhysicsAssetEditor::IsCollisionRenderingMode, EPhysicsAssetEditorRenderMode::None, false));
+		FIsActionChecked::CreateSP(this, &FPhysicsAssetEditor::IsCollisionRenderingMode, EPhysicsAssetEditorCollisionViewMode::None, false));
 
 	ViewportCommandList->EndGroup();
 
@@ -1087,21 +1093,21 @@ void FPhysicsAssetEditor::BindCommands()
 
 	ViewportCommandList->MapAction(
 		Commands.MeshRenderingMode_Simulation_Solid,
-		FExecuteAction::CreateSP(this, &FPhysicsAssetEditor::OnMeshRenderingMode, EPhysicsAssetEditorRenderMode::Solid, true),
+		FExecuteAction::CreateSP(this, &FPhysicsAssetEditor::OnMeshRenderingMode, EPhysicsAssetEditorMeshViewMode::Solid, true),
 		FCanExecuteAction(),
-		FIsActionChecked::CreateSP(this, &FPhysicsAssetEditor::IsMeshRenderingMode, EPhysicsAssetEditorRenderMode::Solid, true));
+		FIsActionChecked::CreateSP(this, &FPhysicsAssetEditor::IsMeshRenderingMode, EPhysicsAssetEditorMeshViewMode::Solid, true));
 
 	ViewportCommandList->MapAction(
 		Commands.MeshRenderingMode_Simulation_Wireframe,
-		FExecuteAction::CreateSP(this, &FPhysicsAssetEditor::OnMeshRenderingMode, EPhysicsAssetEditorRenderMode::Wireframe, true),
+		FExecuteAction::CreateSP(this, &FPhysicsAssetEditor::OnMeshRenderingMode, EPhysicsAssetEditorMeshViewMode::Wireframe, true),
 		FCanExecuteAction(),
-		FIsActionChecked::CreateSP(this, &FPhysicsAssetEditor::IsMeshRenderingMode, EPhysicsAssetEditorRenderMode::Wireframe, true));
+		FIsActionChecked::CreateSP(this, &FPhysicsAssetEditor::IsMeshRenderingMode, EPhysicsAssetEditorMeshViewMode::Wireframe, true));
 
 	ViewportCommandList->MapAction(
 		Commands.MeshRenderingMode_Simulation_None,
-		FExecuteAction::CreateSP(this, &FPhysicsAssetEditor::OnMeshRenderingMode, EPhysicsAssetEditorRenderMode::None, true),
+		FExecuteAction::CreateSP(this, &FPhysicsAssetEditor::OnMeshRenderingMode, EPhysicsAssetEditorMeshViewMode::None, true),
 		FCanExecuteAction(),
-		FIsActionChecked::CreateSP(this, &FPhysicsAssetEditor::IsMeshRenderingMode, EPhysicsAssetEditorRenderMode::None, true));
+		FIsActionChecked::CreateSP(this, &FPhysicsAssetEditor::IsMeshRenderingMode, EPhysicsAssetEditorMeshViewMode::None, true));
 
 	ViewportCommandList->EndGroup();
 
@@ -1109,21 +1115,27 @@ void FPhysicsAssetEditor::BindCommands()
 
 	ViewportCommandList->MapAction(
 		Commands.CollisionRenderingMode_Simulation_Solid,
-		FExecuteAction::CreateSP(this, &FPhysicsAssetEditor::OnCollisionRenderingMode, EPhysicsAssetEditorRenderMode::Solid, true),
+		FExecuteAction::CreateSP(this, &FPhysicsAssetEditor::OnCollisionRenderingMode, EPhysicsAssetEditorCollisionViewMode::Solid, true),
 		FCanExecuteAction(),
-		FIsActionChecked::CreateSP(this, &FPhysicsAssetEditor::IsCollisionRenderingMode, EPhysicsAssetEditorRenderMode::Solid, true));
+		FIsActionChecked::CreateSP(this, &FPhysicsAssetEditor::IsCollisionRenderingMode, EPhysicsAssetEditorCollisionViewMode::Solid, true));
 
 	ViewportCommandList->MapAction(
 		Commands.CollisionRenderingMode_Simulation_Wireframe,
-		FExecuteAction::CreateSP(this, &FPhysicsAssetEditor::OnCollisionRenderingMode, EPhysicsAssetEditorRenderMode::Wireframe, true),
+		FExecuteAction::CreateSP(this, &FPhysicsAssetEditor::OnCollisionRenderingMode, EPhysicsAssetEditorCollisionViewMode::Wireframe, true),
+		FCanExecuteAction(),
+		FIsActionChecked::CreateSP(this, &FPhysicsAssetEditor::IsCollisionRenderingMode, EPhysicsAssetEditorCollisionViewMode::Wireframe, true));
+
+	ViewportCommandList->MapAction(
+		Commands.CollisionRenderingMode_Simulation_SolidWireframe,
+		FExecuteAction::CreateSP(this, &FPhysicsAssetEditor::OnCollisionRenderingMode, EPhysicsAssetEditorCollisionViewMode::SolidWireframe, true),
 		FCanExecuteAction(),
-		FIsActionChecked::CreateSP(this, &FPhysicsAssetEditor::IsCollisionRenderingMode, EPhysicsAssetEditorRenderMode::Wireframe, true));
+		FIsActionChecked::CreateSP(this, &FPhysicsAssetEditor::IsCollisionRenderingMode, EPhysicsAssetEditorCollisionViewMode::SolidWireframe, true));
 
 	ViewportCommandList->MapAction(
 		Commands.CollisionRenderingMode_Simulation_None,
-		FExecuteAction::CreateSP(this, &FPhysicsAssetEditor::OnCollisionRenderingMode, EPhysicsAssetEditorRenderMode::None, true),
+		FExecuteAction::CreateSP(this, &FPhysicsAssetEditor::OnCollisionRenderingMode, EPhysicsAssetEditorCollisionViewMode::None, true),
 		FCanExecuteAction(),
-		FIsActionChecked::CreateSP(this, &FPhysicsAssetEditor::IsCollisionRenderingMode, EPhysicsAssetEditorRenderMode::None, true));
+		FIsActionChecked::CreateSP(this, &FPhysicsAssetEditor::IsCollisionRenderingMode, EPhysicsAssetEditorCollisionViewMode::None, true));
 
 	ViewportCommandList->EndGroup();
 
@@ -2267,7 +2279,7 @@ bool FPhysicsAssetEditor::IsToggleSimulation() const
 	return SharedData->bRunningSimulation;
 }
 
-void FPhysicsAssetEditor::OnMeshRenderingMode(EPhysicsAssetEditorRenderMode Mode, bool bSimulation)
+void FPhysicsAssetEditor::OnMeshRenderingMode(EPhysicsAssetEditorMeshViewMode Mode, bool bSimulation)
 {
 	if (bSimulation)
 	{
@@ -2285,12 +2297,12 @@ void FPhysicsAssetEditor::OnMeshRenderingMode(EPhysicsAssetEditorRenderMode Mode
 	RefreshPreviewViewport();
 }
 
-bool FPhysicsAssetEditor::IsMeshRenderingMode(EPhysicsAssetEditorRenderMode Mode, bool bSimulation) const
+bool FPhysicsAssetEditor::IsMeshRenderingMode(EPhysicsAssetEditorMeshViewMode Mode, bool bSimulation) const
 {
 	return Mode == SharedData->GetCurrentMeshViewMode(bSimulation);
 }
 
-void FPhysicsAssetEditor::OnCollisionRenderingMode(EPhysicsAssetEditorRenderMode Mode, bool bSimulation)
+void FPhysicsAssetEditor::OnCollisionRenderingMode(EPhysicsAssetEditorCollisionViewMode Mode, bool bSimulation)
 {
 	if (bSimulation)
 	{
@@ -2306,7 +2318,7 @@ void FPhysicsAssetEditor::OnCollisionRenderingMode(EPhysicsAssetEditorRenderMode
 	RefreshPreviewViewport();
 }
 
-bool FPhysicsAssetEditor::IsCollisionRenderingMode(EPhysicsAssetEditorRenderMode Mode, bool bSimulation) const
+bool FPhysicsAssetEditor::IsCollisionRenderingMode(EPhysicsAssetEditorCollisionViewMode Mode, bool bSimulation) const
 {
 	return Mode == SharedData->GetCurrentCollisionViewMode(bSimulation);
 }
diff --git a/Engine/Source/Editor/PhysicsAssetEditor/Private/PhysicsAssetEditor.h b/Engine/Source/Editor/PhysicsAssetEditor/Private/PhysicsAssetEditor.h
--- a/Engine/Source/Editor/PhysicsAssetEditor/Private/PhysicsAssetEditor.h
+++ b/Engine/Source/Editor/PhysicsAssetEditor/Private/PhysicsAssetEditor.h
@@ -227,10 +227,10 @@ private:
 	bool IsFullSimulation() const;
 	bool IsSelectedSimulation() const;
 	bool IsToggleSimulation() const;
-	void OnMeshRenderingMode(EPhysicsAssetEditorRenderMode Mode, bool bSimulation);
-	bool IsMeshRenderingMode(EPhysicsAssetEditorRenderMode Mode, bool bSimulation) const;
-	void OnCollisionRenderingMode(EPhysicsAssetEditorRenderMode Mode, bool bSimulation);
-	bool IsCollisionRenderingMode(EPhysicsAssetEditorRenderMode Mode, bool bSimulation) const;
+	void OnMeshRenderingMode(EPhysicsAssetEditorMeshViewMode Mode, bool bSimulation);
+	bool IsMeshRenderingMode(EPhysicsAssetEditorMeshViewMode Mode, bool bSimulation) const;
+	void OnCollisionRenderingMode(EPhysicsAssetEditorCollisionViewMode Mode, bool bSimulation);
+	bool IsCollisionRenderingMode(EPhysicsAssetEditorCollisionViewMode Mode, bool bSimulation) const;
 	void OnConstraintRenderingMode(EPhysicsAssetEditorConstraintViewMode Mode, bool bSimulation);
 	bool IsConstraintRenderingMode(EPhysicsAssetEditorConstraintViewMode Mode, bool bSimulation) const;
 	void ToggleDrawConstraintsAsPoints();
diff --git a/Engine/Source/Editor/PhysicsAssetEditor/Private/PhysicsAssetEditorActions.cpp b/Engine/Source/Editor/PhysicsAssetEditor/Private/PhysicsAssetEditorActions.cpp
--- a/Engine/Source/Editor/PhysicsAssetEditor/Private/PhysicsAssetEditorActions.cpp
+++ b/Engine/Source/Editor/PhysicsAssetEditor/Private/PhysicsAssetEditorActions.cpp
@@ -24,6 +24,7 @@ void FPhysicsAssetEditorCommands::RegisterCommands()
 	UI_COMMAND(MeshRenderingMode_None, "None", "No Mesh Rendering Mode (Edit)", EUserInterfaceActionType::RadioButton, FInputChord());
 	UI_COMMAND(CollisionRenderingMode_Solid, "Solid", "Solid Collision Rendering Mode (Edit)", EUserInterfaceActionType::RadioButton, FInputChord());
 	UI_COMMAND(CollisionRenderingMode_Wireframe, "Wireframe", "Wireframe Collision Rendering Mode (Edit)", EUserInterfaceActionType::RadioButton, FInputChord());
+	UI_COMMAND(CollisionRenderingMode_SolidWireframe, "Solid + Wireframe", "Solid + Wireframe Collision Rendering Mode (Edit)", EUserInterfaceActionType::RadioButton, FInputChord());
 	UI_COMMAND(CollisionRenderingMode_None, "None", "No Collision Rendering Mode (Edit)", EUserInterfaceActionType::RadioButton, FInputChord());
 	UI_COMMAND(ConstraintRenderingMode_None, "None", "No Constraint Rendering Mode (Edit)", EUserInterfaceActionType::RadioButton, FInputChord());
 	UI_COMMAND(ConstraintRenderingMode_AllPositions, "All Positions", "All Positions Constraint Rendering Mode (Edit)", EUserInterfaceActionType::RadioButton, FInputChord());
@@ -34,6 +35,7 @@ void FPhysicsAssetEditorCommands::RegisterCommands()
 	UI_COMMAND(MeshRenderingMode_Simulation_None, "None", "No Mesh Rendering Mode (Simulation)", EUserInterfaceActionType::RadioButton, FInputChord());
 	UI_COMMAND(CollisionRenderingMode_Simulation_Solid, "Solid", "Solid Collision Rendering Mode (Simulation)", EUserInterfaceActionType::RadioButton, FInputChord());
 	UI_COMMAND(CollisionRenderingMode_Simulation_Wireframe, "Wireframe", "Wireframe Collision Rendering Mode (Simulation)", EUserInterfaceActionType::RadioButton, FInputChord());
+	UI_COMMAND(CollisionRenderingMode_Simulation_SolidWireframe, "Solid + Wireframe", "Solid + Wireframe Collision Rendering Mode (Simulation)", EUserInterfaceActionType::RadioButton, FInputChord());
 	UI_COMMAND(CollisionRenderingMode_Simulation_None, "None", "No Collision Rendering Mode (Simulation)", EUserInterfaceActionType::RadioButton, FInputChord());
 	UI_COMMAND(ConstraintRenderingMode_Simulation_None, "None", "No Constraint Rendering Mode (Simulation)", EUserInterfaceActionType::RadioButton, FInputChord());
 	UI_COMMAND(ConstraintRenderingMode_Simulation_AllPositions, "All Positions", "All Positions Constraint Rendering Mode (Simulation)", EUserInterfaceActionType::RadioButton, FInputChord());
diff --git a/Engine/Source/Editor/PhysicsAssetEditor/Private/PhysicsAssetEditorActions.h b/Engine/Source/Editor/PhysicsAssetEditor/Private/PhysicsAssetEditorActions.h
--- a/Engine/Source/Editor/PhysicsAssetEditor/Private/PhysicsAssetEditorActions.h
+++ b/Engine/Source/Editor/PhysicsAssetEditor/Private/PhysicsAssetEditorActions.h
@@ -38,6 +38,7 @@ public:
 	TSharedPtr<FUICommandInfo> MeshRenderingMode_None;
 	TSharedPtr<FUICommandInfo> CollisionRenderingMode_Solid;
 	TSharedPtr<FUICommandInfo> CollisionRenderingMode_Wireframe;
+	TSharedPtr<FUICommandInfo> CollisionRenderingMode_SolidWireframe;
 	TSharedPtr<FUICommandInfo> CollisionRenderingMode_None;
 	TSharedPtr<FUICommandInfo> ConstraintRenderingMode_None;
 	TSharedPtr<FUICommandInfo> ConstraintRenderingMode_AllPositions;
@@ -47,6 +48,7 @@ public:
 	TSharedPtr<FUICommandInfo> MeshRenderingMode_Simulation_None;
 	TSharedPtr<FUICommandInfo> CollisionRenderingMode_Simulation_Solid;
 	TSharedPtr<FUICommandInfo> CollisionRenderingMode_Simulation_Wireframe;
+	TSharedPtr<FUICommandInfo> CollisionRenderingMode_Simulation_SolidWireframe;
 	TSharedPtr<FUICommandInfo> CollisionRenderingMode_Simulation_None;
 	TSharedPtr<FUICommandInfo> ConstraintRenderingMode_Simulation_None;
 	TSharedPtr<FUICommandInfo> ConstraintRenderingMode_Simulation_AllPositions;
diff --git a/Engine/Source/Editor/PhysicsAssetEditor/Private/PhysicsAssetEditorEditMode.cpp b/Engine/Source/Editor/PhysicsAssetEditor/Private/PhysicsAssetEditorEditMode.cpp
--- a/Engine/Source/Editor/PhysicsAssetEditor/Private/PhysicsAssetEditorEditMode.cpp
+++ b/Engine/Source/Editor/PhysicsAssetEditor/Private/PhysicsAssetEditorEditMode.cpp
@@ -31,7 +31,6 @@ FPhysicsAssetEditorEditMode::FPhysicsAssetEditorEditMode()
 	, PhysicsAssetEditor_TranslateSpeed(0.25f)
 	, PhysicsAssetEditor_RotateSpeed(1.0f * (PI / 180.0f))
 	, PhysicsAssetEditor_LightRotSpeed(0.22f)
-	, SimGrabCheckDistance(500.0f)
 	, SimHoldDistanceChangeDelta(20.0f)
 	, SimMinHoldDistance(10.0f)
 	, SimGrabMoveSpeed(1.0f)
@@ -503,13 +502,13 @@ void FPhysicsAssetEditorEditMode::Tick(FEditorViewportClient* ViewportClient, fl
 
 void FPhysicsAssetEditorEditMode::Render(const FSceneView* View, FViewport* Viewport, FPrimitiveDrawInterface* PDI)
 {
-	EPhysicsAssetEditorRenderMode MeshViewMode = SharedData->GetCurrentMeshViewMode(SharedData->bRunningSimulation);
+	EPhysicsAssetEditorMeshViewMode MeshViewMode = SharedData->GetCurrentMeshViewMode(SharedData->bRunningSimulation);
 
-	if (MeshViewMode != EPhysicsAssetEditorRenderMode::None)
+	if (MeshViewMode != EPhysicsAssetEditorMeshViewMode::None)
 	{
 		SharedData->EditorSkelComp->SetVisibility(true);
 
-		if (MeshViewMode == EPhysicsAssetEditorRenderMode::Wireframe)
+		if (MeshViewMode == EPhysicsAssetEditorMeshViewMode::Wireframe)
 		{
 			SharedData->EditorSkelComp->SetForceWireframe(true);
 		}
@@ -829,7 +828,7 @@ bool FPhysicsAssetEditorEditMode::SimMousePress(FEditorViewportClient* InViewpor
 
 	const FViewportClick Click(View, InViewportClient, EKeys::Invalid, IE_Released, Viewport->GetMouseX(), Viewport->GetMouseY());
 	FHitResult Result(1.f);
-	bool bHit = SharedData->EditorSkelComp->LineTraceComponent(Result, Click.GetOrigin(), Click.GetOrigin() + Click.GetDirection() * SimGrabCheckDistance, FCollisionQueryParams(NAME_None, true));
+	bool bHit = SharedData->EditorSkelComp->LineTraceComponent(Result, Click.GetOrigin(), Click.GetOrigin() + Click.GetDirection() * SharedData->EditorOptions->InteractionDistance, FCollisionQueryParams(NAME_None, true));
 
 	SharedData->LastClickPos = Click.GetClickPos();
 	SharedData->LastClickOrigin = Click.GetOrigin();
@@ -870,7 +869,7 @@ bool FPhysicsAssetEditorEditMode::SimMousePress(FEditorViewportClient* InViewpor
 
 				FMatrix	InvViewMatrix = View->ViewMatrices.GetInvViewMatrix();
 
-				SimGrabMinPush = SimMinHoldDistance - (Result.Time * SimGrabCheckDistance);
+				SimGrabMinPush = SimMinHoldDistance - (Result.Time * SharedData->EditorOptions->InteractionDistance);
 
 				SimGrabLocation = Result.Location;
 				SimGrabX = InvViewMatrix.GetUnitAxis(EAxis::X);
diff --git a/Engine/Source/Editor/PhysicsAssetEditor/Private/PhysicsAssetEditorEditMode.h b/Engine/Source/Editor/PhysicsAssetEditor/Private/PhysicsAssetEditorEditMode.h
--- a/Engine/Source/Editor/PhysicsAssetEditor/Private/PhysicsAssetEditorEditMode.h
+++ b/Engine/Source/Editor/PhysicsAssetEditor/Private/PhysicsAssetEditorEditMode.h
@@ -88,7 +88,6 @@ private:
 	const float PhysicsAssetEditor_TranslateSpeed;
 	const float PhysicsAssetEditor_RotateSpeed;
 	const float PhysicsAssetEditor_LightRotSpeed;
-	const float	SimGrabCheckDistance;
 	const float	SimHoldDistanceChangeDelta;
 	const float	SimMinHoldDistance;
 	const float SimGrabMoveSpeed;
diff --git a/Engine/Source/Editor/PhysicsAssetEditor/Private/PhysicsAssetEditorMode.cpp b/Engine/Source/Editor/PhysicsAssetEditor/Private/PhysicsAssetEditorMode.cpp
--- a/Engine/Source/Editor/PhysicsAssetEditor/Private/PhysicsAssetEditorMode.cpp
+++ b/Engine/Source/Editor/PhysicsAssetEditor/Private/PhysicsAssetEditorMode.cpp
@@ -94,6 +94,7 @@ FPhysicsAssetEditorMode::FPhysicsAssetEditorMode(TSharedRef<FWorkflowCentricAppl
 					{
 						InSubMenuBuilder.AddMenuEntry(Commands.CollisionRenderingMode_Solid);
 						InSubMenuBuilder.AddMenuEntry(Commands.CollisionRenderingMode_Wireframe);
+						InSubMenuBuilder.AddMenuEntry(Commands.CollisionRenderingMode_SolidWireframe);
 						InSubMenuBuilder.AddMenuEntry(Commands.CollisionRenderingMode_None);
 					}
 					InSubMenuBuilder.EndSection();
@@ -102,6 +103,7 @@ FPhysicsAssetEditorMode::FPhysicsAssetEditorMode(TSharedRef<FWorkflowCentricAppl
 					{
 						InSubMenuBuilder.AddMenuEntry(Commands.CollisionRenderingMode_Simulation_Solid);
 						InSubMenuBuilder.AddMenuEntry(Commands.CollisionRenderingMode_Simulation_Wireframe);
+						InSubMenuBuilder.AddMenuEntry(Commands.CollisionRenderingMode_Simulation_SolidWireframe);
 						InSubMenuBuilder.AddMenuEntry(Commands.CollisionRenderingMode_Simulation_None);
 					}
 					InSubMenuBuilder.EndSection();
diff --git a/Engine/Source/Editor/PhysicsAssetEditor/Private/PhysicsAssetEditorSharedData.cpp b/Engine/Source/Editor/PhysicsAssetEditor/Private/PhysicsAssetEditorSharedData.cpp
--- a/Engine/Source/Editor/PhysicsAssetEditor/Private/PhysicsAssetEditorSharedData.cpp
+++ b/Engine/Source/Editor/PhysicsAssetEditor/Private/PhysicsAssetEditorSharedData.cpp
@@ -408,7 +408,7 @@ void FPhysicsAssetEditorSharedData::Mirror()
 	}
 }
 
-EPhysicsAssetEditorRenderMode FPhysicsAssetEditorSharedData::GetCurrentMeshViewMode(bool bSimulation)
+EPhysicsAssetEditorMeshViewMode FPhysicsAssetEditorSharedData::GetCurrentMeshViewMode(bool bSimulation)
 {
 	if (bSimulation)
 	{
@@ -420,7 +420,7 @@ EPhysicsAssetEditorRenderMode FPhysicsAssetEditorSharedData::GetCurrentMeshViewM
 	}
 }
 
-EPhysicsAssetEditorRenderMode FPhysicsAssetEditorSharedData::GetCurrentCollisionViewMode(bool bSimulation)
+EPhysicsAssetEditorCollisionViewMode FPhysicsAssetEditorSharedData::GetCurrentCollisionViewMode(bool bSimulation)
 {
 	if (bSimulation)
 	{
diff --git a/Engine/Source/Editor/PhysicsAssetEditor/Private/PhysicsAssetEditorSharedData.h b/Engine/Source/Editor/PhysicsAssetEditor/Private/PhysicsAssetEditorSharedData.h
--- a/Engine/Source/Editor/PhysicsAssetEditor/Private/PhysicsAssetEditorSharedData.h
+++ b/Engine/Source/Editor/PhysicsAssetEditor/Private/PhysicsAssetEditorSharedData.h
@@ -74,10 +74,10 @@ public:
 	void CachePreviewMesh();
 
 	/** Accessor for mesh view mode, allows access for simulation and non-simulation modes */
-	EPhysicsAssetEditorRenderMode GetCurrentMeshViewMode(bool bSimulation);
+	EPhysicsAssetEditorMeshViewMode GetCurrentMeshViewMode(bool bSimulation);
 
 	/** Accessor for collision view mode, allows access for simulation and non-simulation modes */
-	EPhysicsAssetEditorRenderMode GetCurrentCollisionViewMode(bool bSimulation);
+	EPhysicsAssetEditorCollisionViewMode GetCurrentCollisionViewMode(bool bSimulation);
 
 	/** Accessor for constraint view mode, allows access for simulation and non-simulation modes */
 	EPhysicsAssetEditorConstraintViewMode GetCurrentConstraintViewMode(bool bSimulation);
diff --git a/Engine/Source/Editor/PhysicsAssetEditor/Private/PhysicsAssetEditorSkeletalMeshComponent.cpp b/Engine/Source/Editor/PhysicsAssetEditor/Private/PhysicsAssetEditorSkeletalMeshComponent.cpp
--- a/Engine/Source/Editor/PhysicsAssetEditor/Private/PhysicsAssetEditorSkeletalMeshComponent.cpp
+++ b/Engine/Source/Editor/PhysicsAssetEditor/Private/PhysicsAssetEditorSkeletalMeshComponent.cpp
@@ -50,10 +50,6 @@ UPhysicsAssetEditorSkeletalMeshComponent::UPhysicsAssetEditorSkeletalMeshCompone
 		ElemSelectedMaterial = UMaterialInstanceDynamic::Create(BaseElemSelectedMaterial, GetTransientPackage());
 		check(ElemSelectedMaterial);
 
-		UMaterialInterface* BaseBoneSelectedMaterial = LoadObject<UMaterialInterface>(NULL, TEXT("/Engine/EditorMaterials/PhAT_BoneSelectedMaterial.PhAT_BoneSelectedMaterial"), NULL, LOAD_None, NULL);
-		BoneSelectedMaterial = UMaterialInstanceDynamic::Create(BaseBoneSelectedMaterial, GetTransientPackage());
-		check(BoneSelectedMaterial);
-
 		BoneMaterialHit = UMaterial::GetDefaultMaterial(MD_Surface);
 		check(BoneMaterialHit);
 
@@ -89,7 +85,7 @@ void UPhysicsAssetEditorSkeletalMeshComponent::RenderAssetTools(const FSceneView
 		return;
 	}
 
-	EPhysicsAssetEditorRenderMode CollisionViewMode = SharedData->GetCurrentCollisionViewMode(SharedData->bRunningSimulation);
+	EPhysicsAssetEditorCollisionViewMode CollisionViewMode = SharedData->GetCurrentCollisionViewMode(SharedData->bRunningSimulation);
 
 	if (bDebugViewportClicks)
 	{
@@ -102,7 +98,6 @@ void UPhysicsAssetEditorSkeletalMeshComponent::RenderAssetTools(const FSceneView
 	// set opacity of our materials
 	static FName OpacityName(TEXT("Opacity"));
 	ElemSelectedMaterial->SetScalarParameterValue(OpacityName, SharedData->EditorOptions->CollisionOpacity);
-	BoneSelectedMaterial->SetScalarParameterValue(OpacityName, SharedData->EditorOptions->CollisionOpacity);
 	BoneUnselectedMaterial->SetScalarParameterValue(OpacityName, SharedData->EditorOptions->bSolidRenderingForSelectedOnly ? 0.0f : SharedData->EditorOptions->CollisionOpacity);
 	BoneNoCollisionMaterial->SetScalarParameterValue(OpacityName, SharedData->EditorOptions->bSolidRenderingForSelectedOnly ? 0.0f : SharedData->EditorOptions->CollisionOpacity);
 
@@ -111,7 +106,43 @@ void UPhysicsAssetEditorSkeletalMeshComponent::RenderAssetTools(const FSceneView
 	const FLinearColor LinearSelectionColor(SelectionColor.IsColorSpecified() ? SelectionColor.GetSpecifiedColor() : FLinearColor::White);
 
 	ElemSelectedMaterial->SetVectorParameterValue(SelectionColorName, LinearSelectionColor);
-	BoneSelectedMaterial->SetVectorParameterValue(SelectionColorName, LinearSelectionColor);
+
+	// Contains info needed for drawing, in a form that can be sorted by distance.
+	struct DrawElement
+	{
+		DrawElement(
+			FTransform        InTM, 
+			HHitProxy*        InHitProxy, 
+			float             InScale, 
+			const FSceneView* InView,
+			FKShapeElem*      InShapeElem,
+			const FName&      InBoneName
+		)
+			: TM(InTM), HitProxy(InHitProxy), Scale(InScale), ShapeElem(InShapeElem) 
+#ifdef UE_BUILD_DEBUG
+			, BoneName(InBoneName)
+#endif
+		{
+			// The TM position is at the center of the objects, so can be used directly for sorting.
+			// Sorting by distance (rather than along the view direction) reduces flickering when
+			// the camera is rotated without moving it.
+			FVector ViewDir = InView->GetViewDirection();
+			//Distance = TM.GetTranslation() | ViewDir; // Would sort along the view direction
+			SortingMetric = (TM.GetTranslation() - InView->ViewLocation).SquaredLength(); // Sorts by Distance
+		}
+
+		UMaterialInterface*   Material = 0;
+		FColor                Color = FColor(ForceInitToZero);
+		FTransform            TM;
+		HHitProxy*            HitProxy = nullptr;
+		float                 Scale;
+		FKShapeElem*          ShapeElem = nullptr;
+		float                 SortingMetric;
+#ifdef UE_BUILD_DEBUG
+		FName                 BoneName;
+#endif
+	};
+	TArray<DrawElement> DrawElements;
 
 	// Draw bodies
 	for (int32 i = 0; i <PhysicsAsset->SkeletalBodySetups.Num(); ++i)
@@ -140,119 +171,147 @@ void UPhysicsAssetEditorSkeletalMeshComponent::RenderAssetTools(const FSceneView
 		{
 			FTransform BoneTM = GetBoneTransform(BoneIndex);
 			float Scale = BoneTM.GetScale3D().GetAbsMax();
-			FVector VectorScale(Scale);
 			BoneTM.RemoveScaling();
 
 			FKAggregateGeom* AggGeom = &PhysicsAsset->SkeletalBodySetups[i]->AggGeom;
 
 			for (int32 j = 0; j <AggGeom->SphereElems.Num(); ++j)
 			{
-				PDI->SetHitProxy(new HPhysicsAssetEditorEdBoneProxy(i, EAggCollisionShape::Sphere, j));
-
-				FTransform ElemTM = GetPrimitiveTransform(BoneTM, i, EAggCollisionShape::Sphere, j, Scale);
-
+				DrawElement DE(
+					GetPrimitiveTransform(BoneTM, i, EAggCollisionShape::Sphere, j, Scale), 
+					new HPhysicsAssetEditorEdBoneProxy(i, EAggCollisionShape::Sphere, j),
+					Scale, View, &AggGeom->SphereElems[j], GetBoneName(BoneIndex)
+				);
 
 				//solids are drawn if it's the ViewMode and we're not doing a hit, or if it's hitAndBodyMode
-				if (CollisionViewMode == EPhysicsAssetEditorRenderMode::Solid)
+				if (CollisionViewMode == EPhysicsAssetEditorCollisionViewMode::Solid || CollisionViewMode == EPhysicsAssetEditorCollisionViewMode::SolidWireframe)
 				{
-					UMaterialInterface*	PrimMaterial = GetPrimitiveMaterial(i, EAggCollisionShape::Sphere, j);
-					AggGeom->SphereElems[j].DrawElemSolid(PDI, ElemTM, VectorScale, PrimMaterial->GetRenderProxy());
+					DE.Material = GetPrimitiveMaterial(i, EAggCollisionShape::Sphere, j);
 				}
-
-				if (CollisionViewMode == EPhysicsAssetEditorRenderMode::Solid || CollisionViewMode == EPhysicsAssetEditorRenderMode::Wireframe)
+				if (CollisionViewMode == EPhysicsAssetEditorCollisionViewMode::SolidWireframe || CollisionViewMode == EPhysicsAssetEditorCollisionViewMode::Wireframe)
 				{
-					AggGeom->SphereElems[j].DrawElemWire(PDI, ElemTM, VectorScale, GetPrimitiveColor(i, EAggCollisionShape::Sphere, j));
+					DE.Color = GetPrimitiveColor(i, EAggCollisionShape::Sphere, j);
+				}
+				if (DE.Material || DE.Color.A)
+				{
+					DrawElements.Add(DE);
 				}
-
-				PDI->SetHitProxy(NULL);
 			}
 
 			for (int32 j = 0; j <AggGeom->BoxElems.Num(); ++j)
 			{
-				PDI->SetHitProxy(new HPhysicsAssetEditorEdBoneProxy(i, EAggCollisionShape::Box, j));
-
-				FTransform ElemTM = GetPrimitiveTransform(BoneTM, i, EAggCollisionShape::Box, j, Scale);
+				DrawElement DE(
+					GetPrimitiveTransform(BoneTM, i, EAggCollisionShape::Box, j, Scale),
+					new HPhysicsAssetEditorEdBoneProxy(i, EAggCollisionShape::Box, j),
+					Scale, View, &AggGeom->BoxElems[j], GetBoneName(BoneIndex)
+				);
 
-				if (CollisionViewMode == EPhysicsAssetEditorRenderMode::Solid)
+				if (CollisionViewMode == EPhysicsAssetEditorCollisionViewMode::Solid || CollisionViewMode == EPhysicsAssetEditorCollisionViewMode::SolidWireframe)
 				{
-					UMaterialInterface*	PrimMaterial = GetPrimitiveMaterial(i, EAggCollisionShape::Box, j);
-					AggGeom->BoxElems[j].DrawElemSolid(PDI, ElemTM, VectorScale, PrimMaterial->GetRenderProxy());
+					DE.Material = GetPrimitiveMaterial(i, EAggCollisionShape::Box, j);
 				}
-
-				if (CollisionViewMode == EPhysicsAssetEditorRenderMode::Solid || CollisionViewMode == EPhysicsAssetEditorRenderMode::Wireframe)
+				if (CollisionViewMode == EPhysicsAssetEditorCollisionViewMode::SolidWireframe || CollisionViewMode == EPhysicsAssetEditorCollisionViewMode::Wireframe)
 				{
-					AggGeom->BoxElems[j].DrawElemWire(PDI, ElemTM, VectorScale, GetPrimitiveColor(i, EAggCollisionShape::Box, j));
+					DE.Color = GetPrimitiveColor(i, EAggCollisionShape::Box, j);
+				}
+				if (DE.Material || DE.Color.A)
+				{
+					DrawElements.Add(DE);
 				}
-
-				PDI->SetHitProxy(NULL);
 			}
 
 			for (int32 j = 0; j <AggGeom->SphylElems.Num(); ++j)
 			{
-				PDI->SetHitProxy(new HPhysicsAssetEditorEdBoneProxy(i, EAggCollisionShape::Sphyl, j));
-
-				FTransform ElemTM = GetPrimitiveTransform(BoneTM, i, EAggCollisionShape::Sphyl, j, Scale);
+				DrawElement DE(
+					GetPrimitiveTransform(BoneTM, i, EAggCollisionShape::Sphyl, j, Scale),
+					new HPhysicsAssetEditorEdBoneProxy(i, EAggCollisionShape::Sphyl, j),
+					Scale, View, &AggGeom->SphylElems[j], GetBoneName(BoneIndex)
+				);
 
-				if (CollisionViewMode == EPhysicsAssetEditorRenderMode::Solid)
+				if (CollisionViewMode == EPhysicsAssetEditorCollisionViewMode::Solid || CollisionViewMode == EPhysicsAssetEditorCollisionViewMode::SolidWireframe)
 				{
-					UMaterialInterface*	PrimMaterial = GetPrimitiveMaterial(i, EAggCollisionShape::Sphyl, j);
-					AggGeom->SphylElems[j].DrawElemSolid(PDI, ElemTM, VectorScale, PrimMaterial->GetRenderProxy());
+					DE.Material = GetPrimitiveMaterial(i, EAggCollisionShape::Sphyl, j);
 				}
-
-				if (CollisionViewMode == EPhysicsAssetEditorRenderMode::Solid || CollisionViewMode == EPhysicsAssetEditorRenderMode::Wireframe)
+				if (CollisionViewMode == EPhysicsAssetEditorCollisionViewMode::SolidWireframe || CollisionViewMode == EPhysicsAssetEditorCollisionViewMode::Wireframe)
 				{
-					AggGeom->SphylElems[j].DrawElemWire(PDI, ElemTM, VectorScale, GetPrimitiveColor(i, EAggCollisionShape::Sphyl, j));
+					DE.Color = GetPrimitiveColor(i, EAggCollisionShape::Sphyl, j);
+				}
+				if (DE.Material || DE.Color.A)
+				{
+					DrawElements.Add(DE);
 				}
-
-				PDI->SetHitProxy(NULL);
 			}
 
 			for (int32 j = 0; j <AggGeom->ConvexElems.Num(); ++j)
 			{
-				PDI->SetHitProxy(new HPhysicsAssetEditorEdBoneProxy(i, EAggCollisionShape::Convex, j));
-
-				FTransform ElemTM = GetPrimitiveTransform(BoneTM, i, EAggCollisionShape::Convex, j, Scale);
+				DrawElement DE(
+					GetPrimitiveTransform(BoneTM, i, EAggCollisionShape::Convex, j, Scale),
+					new HPhysicsAssetEditorEdBoneProxy(i, EAggCollisionShape::Convex, j),
+					Scale, View, &AggGeom->ConvexElems[j], GetBoneName(BoneIndex)
+				);
 
-				if (CollisionViewMode == EPhysicsAssetEditorRenderMode::Solid)
+				if (CollisionViewMode == EPhysicsAssetEditorCollisionViewMode::Solid || CollisionViewMode == EPhysicsAssetEditorCollisionViewMode::SolidWireframe)
 				{
-					UMaterialInterface* PrimMaterial = GetPrimitiveMaterial(i, EAggCollisionShape::Sphyl, j);
-					AggGeom->ConvexElems[j].DrawElemSolid(PDI, ElemTM, Scale, PrimMaterial->GetRenderProxy());
+					DE.Material = GetPrimitiveMaterial(i, EAggCollisionShape::Convex, j);
 				}
-
-				if (CollisionViewMode == EPhysicsAssetEditorRenderMode::Solid || CollisionViewMode == EPhysicsAssetEditorRenderMode::Wireframe)
+				if (CollisionViewMode == EPhysicsAssetEditorCollisionViewMode::SolidWireframe || CollisionViewMode == EPhysicsAssetEditorCollisionViewMode::Wireframe)
 				{
-					AggGeom->ConvexElems[j].DrawElemWire(PDI, ElemTM, Scale, GetPrimitiveColor(i, EAggCollisionShape::Convex, j));
+					DE.Color = GetPrimitiveColor(i, EAggCollisionShape::Convex, j);
+				}
+				if (DE.Material || DE.Color.A)
+				{
+					DrawElements.Add(DE);
 				}
-
-				PDI->SetHitProxy(NULL);
 			}
 
 			for (int32 j = 0; j <AggGeom->TaperedCapsuleElems.Num(); ++j)
 			{
-				PDI->SetHitProxy(new HPhysicsAssetEditorEdBoneProxy(i, EAggCollisionShape::TaperedCapsule, j));
-
-				FTransform ElemTM = GetPrimitiveTransform(BoneTM, i, EAggCollisionShape::TaperedCapsule, j, Scale);
+				DrawElement DE(
+					GetPrimitiveTransform(BoneTM, i, EAggCollisionShape::TaperedCapsule, j, Scale),
+					new HPhysicsAssetEditorEdBoneProxy(i, EAggCollisionShape::TaperedCapsule, j),
+					Scale, View, &AggGeom->TaperedCapsuleElems[j], GetBoneName(BoneIndex)
+				);
 
-				if (CollisionViewMode == EPhysicsAssetEditorRenderMode::Solid)
+				if (CollisionViewMode == EPhysicsAssetEditorCollisionViewMode::Solid || CollisionViewMode == EPhysicsAssetEditorCollisionViewMode::SolidWireframe)
 				{
-					UMaterialInterface*	PrimMaterial = GetPrimitiveMaterial(i, EAggCollisionShape::TaperedCapsule, j);
-					AggGeom->TaperedCapsuleElems[j].DrawElemSolid(PDI, ElemTM, VectorScale, PrimMaterial->GetRenderProxy());
+					DE.Material = GetPrimitiveMaterial(i, EAggCollisionShape::TaperedCapsule, j);
 				}
-
-				if (CollisionViewMode == EPhysicsAssetEditorRenderMode::Solid || CollisionViewMode == EPhysicsAssetEditorRenderMode::Wireframe)
+				if (CollisionViewMode == EPhysicsAssetEditorCollisionViewMode::SolidWireframe || CollisionViewMode == EPhysicsAssetEditorCollisionViewMode::Wireframe)
+				{
+					DE.Color = GetPrimitiveColor(i, EAggCollisionShape::TaperedCapsule, j);
+				}
+				if (DE.Material || DE.Color.A)
 				{
-					AggGeom->TaperedCapsuleElems[j].DrawElemWire(PDI, ElemTM, VectorScale, GetPrimitiveColor(i, EAggCollisionShape::TaperedCapsule, j));
+					DrawElements.Add(DE);
 				}
-					
-				PDI->SetHitProxy(NULL);
 			}
 
 			if (SharedData->bShowCOM && Bodies.IsValidIndex(i))
 			{
 				Bodies[i]->DrawCOMPosition(PDI, COMRenderSize, SharedData->COMRenderColor);
 			}
+		}
+	}
+
+	// Sort elements
+	DrawElements.Sort([this](const DrawElement& DE1, const DrawElement& DE2) {
+		return DE1.SortingMetric > DE2.SortingMetric; 
+		});
 
+
+	// Render sorted elements.
+	for (const DrawElement& DE : DrawElements)
+	{
+		PDI->SetHitProxy(DE.HitProxy);
+		if (DE.Material)
+		{
+			DE.ShapeElem->DrawElemSolid(PDI, DE.TM, DE.Scale, DE.Material->GetRenderProxy());
+		}
+		if (DE.Color.A)
+		{
+			DE.ShapeElem->DrawElemWire(PDI, DE.TM, DE.Scale, DE.Color);
 		}
+		PDI->SetHitProxy(NULL);
 	}
 
 	// Draw Constraints
@@ -288,8 +347,8 @@ void UPhysicsAssetEditorSkeletalMeshComponent::DebugDraw(const FSceneView* View,
 FPrimitiveSceneProxy* UPhysicsAssetEditorSkeletalMeshComponent::CreateSceneProxy()
 {
 	FPrimitiveSceneProxy* Proxy = NULL;
-	EPhysicsAssetEditorRenderMode MeshViewMode = SharedData->GetCurrentMeshViewMode(SharedData->bRunningSimulation);
-	if (MeshViewMode != EPhysicsAssetEditorRenderMode::None)
+	EPhysicsAssetEditorMeshViewMode MeshViewMode = SharedData->GetCurrentMeshViewMode(SharedData->bRunningSimulation);
+	if (MeshViewMode != EPhysicsAssetEditorMeshViewMode::None)
 	{
 		Proxy = UDebugSkelMeshComponent::CreateSceneProxy();
 	}
@@ -476,7 +535,7 @@ UMaterialInterface* UPhysicsAssetEditorSkeletalMeshComponent::GetPrimitiveMateri
 
 	FPhysicsAssetEditorSharedData::FSelection Body(BodyIndex, PrimitiveType, PrimitiveIndex);
 
-	for (int32 i = 0; i< SharedData->SelectedBodies.Num(); ++i)
+	for (int32 i = 0; i < SharedData->SelectedBodies.Num(); ++i)
 	{
 		if (Body == SharedData->SelectedBodies[i] && !SharedData->bRunningSimulation)
 		{
@@ -484,7 +543,7 @@ UMaterialInterface* UPhysicsAssetEditorSkeletalMeshComponent::GetPrimitiveMateri
 		}
 	}
 
-	if(PrimitiveType == EAggCollisionShape::TaperedCapsule)
+	if (PrimitiveType == EAggCollisionShape::TaperedCapsule)
 	{
 		return BoneNoCollisionMaterial;
 	}
@@ -498,7 +557,6 @@ UMaterialInterface* UPhysicsAssetEditorSkeletalMeshComponent::GetPrimitiveMateri
 	{
 		return BoneUnselectedMaterial;
 	}
-
 }
 
 void UPhysicsAssetEditorSkeletalMeshComponent::RefreshBoneTransforms(FActorComponentTickFunction* TickFunction)
diff --git a/Engine/Source/Editor/PhysicsAssetEditor/Private/PhysicsAssetEditorSkeletalMeshComponent.h b/Engine/Source/Editor/PhysicsAssetEditor/Private/PhysicsAssetEditorSkeletalMeshComponent.h
--- a/Engine/Source/Editor/PhysicsAssetEditor/Private/PhysicsAssetEditorSkeletalMeshComponent.h
+++ b/Engine/Source/Editor/PhysicsAssetEditor/Private/PhysicsAssetEditorSkeletalMeshComponent.h
@@ -38,18 +38,16 @@ class UPhysicsAssetEditorSkeletalMeshComponent : public UDebugSkelMeshComponent
 	UPROPERTY(transient)
 	TObjectPtr<UMaterialInstanceDynamic> ElemSelectedMaterial;
 	UPROPERTY(transient)
-	TObjectPtr<UMaterialInstanceDynamic> BoneSelectedMaterial;
-	UPROPERTY(transient)
 	TObjectPtr<UMaterialInstanceDynamic> BoneUnselectedMaterial;
 	UPROPERTY(transient)
-	TObjectPtr<UMaterialInterface> BoneMaterialHit;
-	UPROPERTY(transient)
 	TObjectPtr<UMaterialInstanceDynamic> BoneNoCollisionMaterial;
 
+	UPROPERTY(transient)
+	TObjectPtr<UMaterialInterface> BoneMaterialHit;
+
 	/** Mesh-space matrices showing state of just animation (ie before physics) - useful for debugging! */
 	TArray<FTransform> AnimationSpaceBases;
 
-
 	/** UPrimitiveComponent interface */
 	virtual FPrimitiveSceneProxy* CreateSceneProxy() override;
 	virtual void AddImpulseAtLocation(FVector Impulse, FVector Location, FName BoneName = NAME_None) override;
diff --git a/Engine/Source/Editor/UnrealEd/Classes/Preferences/PhysicsAssetEditorOptions.h b/Engine/Source/Editor/UnrealEd/Classes/Preferences/PhysicsAssetEditorOptions.h
--- a/Engine/Source/Editor/UnrealEd/Classes/Preferences/PhysicsAssetEditorOptions.h
+++ b/Engine/Source/Editor/UnrealEd/Classes/Preferences/PhysicsAssetEditorOptions.h
@@ -10,7 +10,16 @@
 #include "PhysicsAssetEditorOptions.generated.h"
 
 UENUM()
-enum class EPhysicsAssetEditorRenderMode : uint8
+enum class EPhysicsAssetEditorCollisionViewMode : uint8
+{
+	Solid,
+	Wireframe,
+	SolidWireframe,
+	None
+};
+
+UENUM()
+enum class EPhysicsAssetEditorMeshViewMode : uint8
 {
 	Solid,
 	Wireframe,
@@ -43,11 +52,11 @@ class UNREALED_API UPhysicsAssetEditorOptions : public UObject
 	TEnumAsByte<EPhysicsTransformUpdateMode::Type> PhysicsUpdateMode;
 
 	/** Time between poking ragdoll and starting to blend back. */
-	UPROPERTY(EditAnywhere, config, Category=Anim)
+	UPROPERTY(EditAnywhere, config, Category=Anim, meta = (ClampMin = 0))
 	float PokePauseTime;
 
 	/** Time taken to blend from physics to animation. */
-	UPROPERTY(EditAnywhere, config, Category=Anim)
+	UPROPERTY(EditAnywhere, config, Category=Anim, meta = (ClampMin = 0))
 	float PokeBlendTime;
 	
 	/** Scale factor for the gravity used in the simulation */
@@ -67,29 +76,33 @@ class UNREALED_API UPhysicsAssetEditorOptions : public UObject
 	int32 MaxFPS;
 
 	/** Linear damping of mouse spring forces */
-	UPROPERTY(EditAnywhere, config, Category=MouseSpring)
+	UPROPERTY(EditAnywhere, config, Category=MouseSpring, meta = (ClampMin = 0))
 	float HandleLinearDamping;
 
 	/** Linear stiffness of mouse spring forces */
-	UPROPERTY(EditAnywhere, config, Category=MouseSpring)
+	UPROPERTY(EditAnywhere, config, Category=MouseSpring, meta = (ClampMin = 0))
 	float HandleLinearStiffness;
 
 	/** Angular damping of mouse spring forces */
-	UPROPERTY(EditAnywhere, config, Category=MouseSpring)
+	UPROPERTY(EditAnywhere, config, Category=MouseSpring, meta = (ClampMin = 0))
 	float HandleAngularDamping;
 
 	/** Angular stiffness of mouse spring forces */
-	UPROPERTY(EditAnywhere, config, Category=MouseSpring)
+	UPROPERTY(EditAnywhere, config, Category=MouseSpring, meta = (ClampMin = 0))
 	float HandleAngularStiffness;
 
 	/** How quickly we interpolate the physics target transform for mouse spring forces */
-	UPROPERTY(EditAnywhere, config, Category=MouseSpring)
+	UPROPERTY(EditAnywhere, config, Category=MouseSpring, meta = (ClampMin = 0))
 	float InterpolationSpeed;
 
 	/** Strength of the impulse used when poking with left mouse button */
 	UPROPERTY(EditAnywhere, config, Category=Poking)
 	float PokeStrength;
 
+	/** Raycast distance when poking or grabbing */
+	UPROPERTY(EditAnywhere, config, Category = Poking, meta = (ClampMin = 0))
+	float InteractionDistance;
+
 	/** Whether to draw constraints as points */
 	UPROPERTY(config)
 	uint32 bShowConstraintsAsPoints:1;
@@ -104,11 +117,11 @@ class UNREALED_API UPhysicsAssetEditorOptions : public UObject
 
 	/** View mode for meshes in edit mode */
 	UPROPERTY(config)
-	EPhysicsAssetEditorRenderMode MeshViewMode;
+		EPhysicsAssetEditorMeshViewMode MeshViewMode;
 
 	/** View mode for collision in edit mode */
 	UPROPERTY(config)
-	EPhysicsAssetEditorRenderMode CollisionViewMode;
+	EPhysicsAssetEditorCollisionViewMode CollisionViewMode;
 
 	/** View mode for constraints in edit mode */
 	UPROPERTY(config)
@@ -116,18 +129,18 @@ class UNREALED_API UPhysicsAssetEditorOptions : public UObject
 
 	/** View mode for meshes in simulation mode */
 	UPROPERTY(config)
-	EPhysicsAssetEditorRenderMode SimulationMeshViewMode;
+	EPhysicsAssetEditorMeshViewMode SimulationMeshViewMode;
 
 	/** View mode for collision in simulation mode */
 	UPROPERTY(config)
-	EPhysicsAssetEditorRenderMode SimulationCollisionViewMode;
+	EPhysicsAssetEditorCollisionViewMode SimulationCollisionViewMode;
 
 	/** View mode for constraints in simulation mode */
 	UPROPERTY(config)
 	EPhysicsAssetEditorConstraintViewMode SimulationConstraintViewMode;
 
 	/** Opacity of 'solid' rendering */
-	UPROPERTY(config)
+	UPROPERTY(config, meta = (ClampMin = 0, ClampMax = 1))
 	float CollisionOpacity;
 
 	/** When set, turns opacity of solid rendering for unselected bodies to zero */
diff --git a/Engine/Source/Editor/UnrealEd/Private/PreferenceStubs.cpp b/Engine/Source/Editor/UnrealEd/Private/PreferenceStubs.cpp
--- a/Engine/Source/Editor/UnrealEd/Private/PreferenceStubs.cpp
+++ b/Engine/Source/Editor/UnrealEd/Private/PreferenceStubs.cpp
@@ -43,11 +43,11 @@ UPhysicsAssetEditorOptions::UPhysicsAssetEditorOptions(const FObjectInitializer&
 	ConstraintDrawSize = 1.0f;
 
 	// view options
-	MeshViewMode = EPhysicsAssetEditorRenderMode::Solid;
-	CollisionViewMode = EPhysicsAssetEditorRenderMode::Solid;
+	MeshViewMode = EPhysicsAssetEditorMeshViewMode::Solid;
+	CollisionViewMode = EPhysicsAssetEditorCollisionViewMode::Solid;
 	ConstraintViewMode = EPhysicsAssetEditorConstraintViewMode::AllLimits;
-	SimulationMeshViewMode = EPhysicsAssetEditorRenderMode::Solid;
-	SimulationCollisionViewMode = EPhysicsAssetEditorRenderMode::Solid;
+	SimulationMeshViewMode = EPhysicsAssetEditorMeshViewMode::Solid;
+	SimulationCollisionViewMode = EPhysicsAssetEditorCollisionViewMode::Solid;
 	SimulationConstraintViewMode = EPhysicsAssetEditorConstraintViewMode::None;
 
 	CollisionOpacity = 0.3f;
diff --git a/Engine/Source/Runtime/Engine/Classes/PhysicsEngine/BoxElem.h b/Engine/Source/Runtime/Engine/Classes/PhysicsEngine/BoxElem.h
--- a/Engine/Source/Runtime/Engine/Classes/PhysicsEngine/BoxElem.h
+++ b/Engine/Source/Runtime/Engine/Classes/PhysicsEngine/BoxElem.h
@@ -111,8 +111,11 @@ struct FKBoxElem : public FKShapeElem
 
 	FORCEINLINE FVector::FReal GetVolume(const FVector& Scale3D) const { FVector::FReal MinScale = Scale3D.GetMin(); return (X * MinScale) * (Y * MinScale) * (Z * MinScale); }
 
-	ENGINE_API void	DrawElemWire(class FPrimitiveDrawInterface* PDI, const FTransform& ElemTM, const FVector& Scale3D, const FColor Color) const;
-	ENGINE_API void	DrawElemSolid(class FPrimitiveDrawInterface* PDI, const FTransform& ElemTM, const FVector& Scale3D, const FMaterialRenderProxy* MaterialRenderProxy) const;
+	ENGINE_API void DrawElemWire(class FPrimitiveDrawInterface* PDI, const FTransform& ElemTM, float Scale, const FColor Color) const override;
+	ENGINE_API void DrawElemSolid(class FPrimitiveDrawInterface* PDI, const FTransform& ElemTM, float Scale, const FMaterialRenderProxy* MaterialRenderProxy) const override;
+
+	ENGINE_API void DrawElemWire(class FPrimitiveDrawInterface* PDI, const FTransform& ElemTM, const FVector& Scale3D, const FColor Color) const;
+	ENGINE_API void DrawElemSolid(class FPrimitiveDrawInterface* PDI, const FTransform& ElemTM, const FVector& Scale3D, const FMaterialRenderProxy* MaterialRenderProxy) const;
 	
 	ENGINE_API void GetElemSolid(const FTransform& ElemTM, const FVector& Scale3D, const FMaterialRenderProxy* MaterialRenderProxy, int32 ViewIndex, FMeshElementCollector& Collector) const;
 
diff --git a/Engine/Source/Runtime/Engine/Classes/PhysicsEngine/ConvexElem.h b/Engine/Source/Runtime/Engine/Classes/PhysicsEngine/ConvexElem.h
--- a/Engine/Source/Runtime/Engine/Classes/PhysicsEngine/ConvexElem.h
+++ b/Engine/Source/Runtime/Engine/Classes/PhysicsEngine/ConvexElem.h
@@ -71,9 +71,8 @@ public:
 
 	ENGINE_API const FKConvexElem& operator=(const FKConvexElem& Other);
 
-	ENGINE_API void	DrawElemWire(class FPrimitiveDrawInterface* PDI, const FTransform& ElemTM, const float Scale, const FColor Color) const;
-
-	ENGINE_API void DrawElemSolid(class FPrimitiveDrawInterface* PDI, const FTransform& ElemTM, const float Scale, const FMaterialRenderProxy* MaterialRenderProxy) const;
+	ENGINE_API void DrawElemWire(class FPrimitiveDrawInterface* PDI, const FTransform& ElemTM, const float Scale, const FColor Color) const override;
+	ENGINE_API void DrawElemSolid(class FPrimitiveDrawInterface* PDI, const FTransform& ElemTM, const float Scale, const FMaterialRenderProxy* MaterialRenderProxy) const override;
 
 	ENGINE_API void AddCachedSolidConvexGeom(TArray<FDynamicMeshVertex>& VertexBuffer, TArray<uint32>& IndexBuffer, const FColor VertexColor) const;
 
diff --git a/Engine/Source/Runtime/Engine/Classes/PhysicsEngine/ShapeElem.h b/Engine/Source/Runtime/Engine/Classes/PhysicsEngine/ShapeElem.h
--- a/Engine/Source/Runtime/Engine/Classes/PhysicsEngine/ShapeElem.h
+++ b/Engine/Source/Runtime/Engine/Classes/PhysicsEngine/ShapeElem.h
@@ -104,6 +104,9 @@ struct FKShapeElem
 	/** Get whether this shape should be considered for query or sim collision */
 	ENGINE_API ECollisionEnabled::Type GetCollisionEnabled() const { return CollisionEnabled; }
 
+	ENGINE_API virtual void DrawElemWire(class FPrimitiveDrawInterface* PDI, const FTransform& ElemTM, const float Scale, const FColor Color) const {}
+	ENGINE_API virtual void DrawElemSolid(class FPrimitiveDrawInterface* PDI, const FTransform& ElemTM, const float Scale, const class FMaterialRenderProxy* MaterialRenderProxy) const {}
+
 	/** Offset used when generating contact points. This allows you to smooth out
 		the Minkowski sum by radius R. Useful for making objects slide smoothly
 		on top of irregularities  */
diff --git a/Engine/Source/Runtime/Engine/Classes/PhysicsEngine/SphereElem.h b/Engine/Source/Runtime/Engine/Classes/PhysicsEngine/SphereElem.h
--- a/Engine/Source/Runtime/Engine/Classes/PhysicsEngine/SphereElem.h
+++ b/Engine/Source/Runtime/Engine/Classes/PhysicsEngine/SphereElem.h
@@ -74,9 +74,12 @@ struct FKSphereElem : public FKShapeElem
 
 	FORCEINLINE FVector::FReal GetVolume(const FVector& Scale) const { return 1.3333f * PI * FMath::Pow(Radius * Scale.GetMin(), 3); }
 	
-	ENGINE_API void	DrawElemWire(class FPrimitiveDrawInterface* PDI, const FTransform& ElemTM, const FVector& Scale3D, const FColor Color) const;
+	ENGINE_API void DrawElemWire(class FPrimitiveDrawInterface* PDI, const FTransform& ElemTM, float Scale, const FColor Color) const override;
+	ENGINE_API void DrawElemSolid(class FPrimitiveDrawInterface* PDI, const FTransform& ElemTM, float Scale, const FMaterialRenderProxy* MaterialRenderProxy) const override;
+
+	ENGINE_API void DrawElemWire(class FPrimitiveDrawInterface* PDI, const FTransform& ElemTM, const FVector& Scale3D, const FColor Color) const;
+	ENGINE_API void DrawElemSolid(class FPrimitiveDrawInterface* PDI, const FTransform& ElemTM, const FVector& Scale3D, const FMaterialRenderProxy* MaterialRenderProxy) const;
 
-	ENGINE_API void	DrawElemSolid(class FPrimitiveDrawInterface* PDI, const FTransform& ElemTM, const FVector& Scale3D, const FMaterialRenderProxy* MaterialRenderProxy) const;
 	ENGINE_API void GetElemSolid(const FTransform& ElemTM, const FVector& Scale3D, const FMaterialRenderProxy* MaterialRenderProxy, int32 ViewIndex, class FMeshElementCollector& Collector) const;
 	ENGINE_API FBox CalcAABB(const FTransform& BoneTM, float Scale) const;
 
diff --git a/Engine/Source/Runtime/Engine/Classes/PhysicsEngine/SphylElem.h b/Engine/Source/Runtime/Engine/Classes/PhysicsEngine/SphylElem.h
--- a/Engine/Source/Runtime/Engine/Classes/PhysicsEngine/SphylElem.h
+++ b/Engine/Source/Runtime/Engine/Classes/PhysicsEngine/SphylElem.h
@@ -92,9 +92,13 @@ struct FKSphylElem : public FKShapeElem
 
 	FORCEINLINE FVector::FReal GetVolume(const FVector& Scale) const { FVector::FReal ScaledRadius = Radius * Scale.GetMin(); return PI * FMath::Square(ScaledRadius) * ( 1.3333f * ScaledRadius + (Length * Scale.GetMin())); }
 
-	ENGINE_API void	DrawElemWire(class FPrimitiveDrawInterface* PDI, const FTransform& ElemTM, const FVector& Scale3D, const FColor Color) const;
+	ENGINE_API void DrawElemWire(class FPrimitiveDrawInterface* PDI, const FTransform& ElemTM, float Scale, const FColor Color) const override;
+	ENGINE_API void DrawElemSolid(class FPrimitiveDrawInterface* PDI, const FTransform& ElemTM, float Scale, const FMaterialRenderProxy* MaterialRenderProxy) const override;
+
+	ENGINE_API void DrawElemWire(class FPrimitiveDrawInterface* PDI, const FTransform& ElemTM, const FVector& Scale3D, const FColor Color) const;
+	ENGINE_API void DrawElemSolid(class FPrimitiveDrawInterface* PDI, const FTransform& ElemTM, const FVector& Scale3D, const FMaterialRenderProxy* MaterialRenderProxy) const;
+
 	ENGINE_API void GetElemSolid(const FTransform& ElemTM, const FVector& Scale3D, const FMaterialRenderProxy* MaterialRenderProxy, int32 ViewIndex, FMeshElementCollector& Collector) const;
-	ENGINE_API void	DrawElemSolid(class FPrimitiveDrawInterface* PDI, const FTransform& ElemTM, const FVector& Scale3D, const FMaterialRenderProxy* MaterialRenderProxy) const;
 	ENGINE_API FBox CalcAABB(const FTransform& BoneTM, float Scale) const;
 
 	ENGINE_API void ScaleElem(FVector DeltaSize, float MinSize);
diff --git a/Engine/Source/Runtime/Engine/Classes/PhysicsEngine/TaperedCapsuleElem.h b/Engine/Source/Runtime/Engine/Classes/PhysicsEngine/TaperedCapsuleElem.h
--- a/Engine/Source/Runtime/Engine/Classes/PhysicsEngine/TaperedCapsuleElem.h
+++ b/Engine/Source/Runtime/Engine/Classes/PhysicsEngine/TaperedCapsuleElem.h
@@ -78,9 +78,13 @@ struct FKTaperedCapsuleElem : public FKShapeElem
 		Center = InTransform.GetLocation();
 	}
 
-	ENGINE_API void	DrawElemWire(class FPrimitiveDrawInterface* PDI, const FTransform& ElemTM, const FVector& Scale3D, const FColor Color) const;
+	ENGINE_API void DrawElemWire(class FPrimitiveDrawInterface* PDI, const FTransform& ElemTM, float Scale, const FColor Color) const override;
+	ENGINE_API void DrawElemSolid(class FPrimitiveDrawInterface* PDI, const FTransform& ElemTM, float Scale, const FMaterialRenderProxy* MaterialRenderProxy) const override;
+
+	ENGINE_API void DrawElemWire(class FPrimitiveDrawInterface* PDI, const FTransform& ElemTM, const FVector& Scale3D, const FColor Color) const;
+	ENGINE_API void DrawElemSolid(class FPrimitiveDrawInterface* PDI, const FTransform& ElemTM, const FVector& Scale3D, const FMaterialRenderProxy* MaterialRenderProxy) const;
+
 	ENGINE_API void GetElemSolid(const FTransform& ElemTM, const FVector& Scale3D, const FMaterialRenderProxy* MaterialRenderProxy, int32 ViewIndex, FMeshElementCollector& Collector) const;
-	ENGINE_API void	DrawElemSolid(class FPrimitiveDrawInterface* PDI, const FTransform& ElemTM, const FVector& Scale3D, const FMaterialRenderProxy* MaterialRenderProxy) const;
 	ENGINE_API FBox CalcAABB(const FTransform& BoneTM, float Scale) const;
 
 	ENGINE_API void ScaleElem(FVector DeltaSize, float MinSize);
diff --git a/Engine/Source/Runtime/Engine/Private/PhysicsEngine/PhysDrawing.cpp b/Engine/Source/Runtime/Engine/Private/PhysicsEngine/PhysDrawing.cpp
--- a/Engine/Source/Runtime/Engine/Private/PhysicsEngine/PhysDrawing.cpp
+++ b/Engine/Source/Runtime/Engine/Private/PhysicsEngine/PhysDrawing.cpp
@@ -49,6 +49,16 @@ static const FColor JointLockedColor(255,128,10);
 // FKSphereElem
 /////////////////////////////////////////////////////////////////////////////////////
 
+void FKSphereElem::DrawElemWire(class FPrimitiveDrawInterface* PDI, const FTransform& ElemTM, float Scale, const FColor Color) const
+{
+	DrawElemWire(PDI, ElemTM, FVector(Scale), Color);
+}
+
+void FKSphereElem::DrawElemSolid(class FPrimitiveDrawInterface* PDI, const FTransform& ElemTM, float Scale, const FMaterialRenderProxy* MaterialRenderProxy) const
+{
+	DrawElemSolid(PDI, ElemTM, FVector(Scale), MaterialRenderProxy);
+}
+
 void FKSphereElem::DrawElemWire(class FPrimitiveDrawInterface* PDI, const FTransform& ElemTM, const FVector& Scale3D, const FColor Color) const
 {
 	FVector ElemCenter = ElemTM.GetLocation();
@@ -74,14 +84,28 @@ void FKSphereElem::GetElemSolid(const FTransform& ElemTM, const FVector& Scale3D
 	GetSphereMesh(ElemTM.GetLocation(), FVector(this->Radius * Scale3D.GetAbsMin()), DrawCollisionSides, DrawCollisionSides / 2, MaterialRenderProxy, SDPG_World, false, ViewIndex, Collector);
 }
 
+/////////////////////////////////////////////////////////////////////////////////////
+// FKBoxElem
+/////////////////////////////////////////////////////////////////////////////////////
+
+void FKBoxElem::DrawElemWire(class FPrimitiveDrawInterface* PDI, const FTransform& ElemTM, float Scale, const FColor Color) const
+{
+	DrawElemWire(PDI, ElemTM, FVector(Scale), Color);
+}
+
+void FKBoxElem::DrawElemSolid(class FPrimitiveDrawInterface* PDI, const FTransform& ElemTM, float Scale, const FMaterialRenderProxy* MaterialRenderProxy) const
+{
+	DrawElemSolid(PDI, ElemTM, FVector(Scale), MaterialRenderProxy);
+}
+
 void FKBoxElem::DrawElemWire(class FPrimitiveDrawInterface* PDI, const FTransform& ElemTM, const FVector& Scale3D, const FColor Color) const
 {
 	FVector	B[2], P, Q, Radii;
 
 	// X,Y,Z member variables are LENGTH not RADIUS
-	Radii.X = Scale3D.X*0.5f*X;
-	Radii.Y = Scale3D.Y*0.5f*Y;
-	Radii.Z = Scale3D.Z*0.5f*Z;
+	Radii.X = Scale3D.X * 0.5f * X;
+	Radii.Y = Scale3D.Y * 0.5f * Y;
+	Radii.Z = Scale3D.Z * 0.5f * Z;
 
 	B[0] = Radii; // max
 	B[1] = -1.0f * Radii; // min
@@ -108,10 +132,6 @@ void FKBoxElem::DrawElemWire(class FPrimitiveDrawInterface* PDI, const FTransfor
 	}
 }
 
-/////////////////////////////////////////////////////////////////////////////////////
-// FKBoxElem
-/////////////////////////////////////////////////////////////////////////////////////
-
 void FKBoxElem::DrawElemSolid(class FPrimitiveDrawInterface* PDI, const FTransform& ElemTM, const FVector& Scale3D, const FMaterialRenderProxy* MaterialRenderProxy) const
 {
 	DrawBox(PDI, ElemTM.ToMatrixWithScale(), Scale3D * 0.5f * FVector(X, Y, Z), MaterialRenderProxy, SDPG_World);
@@ -139,6 +159,16 @@ static void DrawHalfCircle(FPrimitiveDrawInterface* PDI, const FVector& Base, co
 	}	
 }
 
+void FKSphylElem::DrawElemWire(class FPrimitiveDrawInterface* PDI, const FTransform& ElemTM, float Scale, const FColor Color) const
+{
+	DrawElemWire(PDI, ElemTM, FVector(Scale), Color);
+}
+
+void FKSphylElem::DrawElemSolid(class FPrimitiveDrawInterface* PDI, const FTransform& ElemTM, float Scale, const FMaterialRenderProxy* MaterialRenderProxy) const
+{
+	DrawElemSolid(PDI, ElemTM, FVector(Scale), MaterialRenderProxy);
+}
+
 void FKSphylElem::DrawElemWire(class FPrimitiveDrawInterface* PDI, const FTransform& ElemTM, const FVector& Scale3D, const FColor Color) const
 {
 	const FVector Origin = ElemTM.GetLocation();
@@ -371,7 +401,7 @@ void FKSphylElem::DrawElemSolid(FPrimitiveDrawInterface* PDI, const FTransform&
 		}
 
 	}
-	MeshBuilder.Draw(PDI, ElemTM.ToMatrixWithScale(), MaterialRenderProxy, SDPG_World,0.f);
+	MeshBuilder.Draw(PDI, ElemTM.ToMatrixWithScale(), MaterialRenderProxy, SDPG_World);
 
 
 	FMemory::Free(Verts);
@@ -417,6 +447,16 @@ void FKTaperedCapsuleElem::DrawTaperedCapsuleSides(FPrimitiveDrawInterface* PDI,
 	}
 }
 
+void FKTaperedCapsuleElem::DrawElemWire(class FPrimitiveDrawInterface* PDI, const FTransform& ElemTM, float Scale, const FColor Color) const
+{
+	DrawElemWire(PDI, ElemTM, FVector(Scale), Color);
+}
+
+void FKTaperedCapsuleElem::DrawElemSolid(class FPrimitiveDrawInterface* PDI, const FTransform& ElemTM, float Scale, const FMaterialRenderProxy* MaterialRenderProxy) const
+{
+	DrawElemSolid(PDI, ElemTM, FVector(Scale), MaterialRenderProxy);
+}
+
 void FKTaperedCapsuleElem::DrawElemWire(class FPrimitiveDrawInterface* PDI, const FTransform& ElemTM, const FVector& Scale3D, const FColor Color) const
 {
 	const FVector Origin = ElemTM.GetLocation();
@@ -627,7 +667,7 @@ void FKTaperedCapsuleElem::GetElemSolid(const FTransform& ElemTM, const FVector&
 void FKTaperedCapsuleElem::DrawElemSolid(FPrimitiveDrawInterface* PDI, const FTransform& ElemTM, const FVector& Scale3D, const FMaterialRenderProxy* MaterialRenderProxy) const
 {
 	FScopedTaperedCapsuleBuilder TaperedCapsuleBuilder(*this, Scale3D, PDI->View->GetFeatureLevel());
-	TaperedCapsuleBuilder.MeshBuilder.Draw(PDI, ElemTM.ToMatrixWithScale(), MaterialRenderProxy, SDPG_World, 0.f);
+	TaperedCapsuleBuilder.MeshBuilder.Draw(PDI, ElemTM.ToMatrixWithScale(), MaterialRenderProxy, SDPG_World);
 }
 
 
@@ -752,7 +792,7 @@ void FKConvexElem::DrawElemSolid(FPrimitiveDrawInterface* PDI, const FTransform&
 			}
 			MeshBuilder.AddTriangle(IndexData[Base], IndexData[Base + 1], IndexData[Base + 2]);
 		}
-		MeshBuilder.Draw(PDI, FMatrix::Identity, MaterialRenderProxy, SDPG_World, 0.f);
+		MeshBuilder.Draw(PDI, FMatrix::Identity, MaterialRenderProxy, SDPG_World);
 	}
 	else
 	{
@@ -953,7 +993,7 @@ void FKAggregateGeom::GetAggGeom(const FTransform& Transform, const FColor Color
 			// Cache collision vertex/index buffer
 			if(!RenderInfo)
 			{
-				//@todo - parallelrendering, remove const cast
+				//@todo - parallel rendering, remove const cast
 				FKAggregateGeom& ThisGeom = const_cast<FKAggregateGeom&>(*this);
 				ThisGeom.RenderInfo = new FKConvexGeomRenderInfo();
 				ThisGeom.RenderInfo->VertexBuffers = new FStaticMeshVertexBuffers();
