Change PhAT simulate key binding to Alt + I

diff --git a/Engine/Source/Editor/PhysicsAssetEditor/Private/PhysicsAssetEditorActions.cpp b/Engine/Source/Editor/PhysicsAssetEditor/Private/PhysicsAssetEditorActions.cpp
--- a/Engine/Source/Editor/PhysicsAssetEditor/Private/PhysicsAssetEditorActions.cpp
+++ b/Engine/Source/Editor/PhysicsAssetEditor/Private/PhysicsAssetEditorActions.cpp
@@ -16,8 +16,8 @@ void FPhysicsAssetEditorCommands::RegisterCommands()
 	UI_COMMAND(PasteProperties, "Paste Properties", "Paste Properties: Copy Properties Of Currently Selected Object To Next Selected Object", EUserInterfaceActionType::Button, FInputChord(EModifierKey::Control, EKeys::V));
 	UI_COMMAND(RepeatLastSimulation, "Simulate", "Previews Physics Simulation", EUserInterfaceActionType::RadioButton, FInputChord());
 	UI_COMMAND(SimulationNoGravity, "No Gravity Simulation", "Run Physics Simulation without gravity. Use this to debug issues with your ragdoll. If the setup is correct, the asset should not move!", EUserInterfaceActionType::ToggleButton, FInputChord());
-	UI_COMMAND(SelectedSimulation, "Simulate Selected", "Run Physics Simulation on selected objects. Use this to tune  specific parts of your ragdoll.", EUserInterfaceActionType::RadioButton, FInputChord(EModifierKey::Alt | EModifierKey::Shift, EKeys::Enter));
-	UI_COMMAND(SimulationAll, "Simulate", "Run Physics Simulation on all objects.", EUserInterfaceActionType::RadioButton, FInputChord(EModifierKey::Alt, EKeys::Enter));
+	UI_COMMAND(SelectedSimulation, "Simulate Selected", "Run Physics Simulation on selected objects. Use this to tune  specific parts of your ragdoll.", EUserInterfaceActionType::RadioButton, FInputChord(EModifierKey::Alt | EModifierKey::Shift, EKeys::I));
+	UI_COMMAND(SimulationAll, "Simulate", "Run Physics Simulation on all objects.", EUserInterfaceActionType::RadioButton, FInputChord(EModifierKey::Alt, EKeys::I));
 
 	UI_COMMAND(MeshRenderingMode_Solid, "Solid", "Solid Mesh Rendering Mode (Edit)", EUserInterfaceActionType::RadioButton, FInputChord());
 	UI_COMMAND(MeshRenderingMode_Wireframe, "Wireframe", "Wireframe Mesh Rendering Mode (Edit)", EUserInterfaceActionType::RadioButton, FInputChord());
@@ -63,7 +63,7 @@ void FPhysicsAssetEditorCommands::RegisterCommands()
 	UI_COMMAND(PrimitiveNoCollision, "No Collision", "Disable all collision on Selected primitive(s)", EUserInterfaceActionType::RadioButton, FInputChord());
 	UI_COMMAND(PrimitiveContributeToMass, "Primitive Contributes To Mass", "Toggle the contribution of selected primitive's volume to the overall mass of the body", EUserInterfaceActionType::Check, FInputChord());
 	UI_COMMAND(DuplicatePrimitive, "Duplicate", "Duplicate Selected Primitive(s)", EUserInterfaceActionType::Button, FInputChord());
-	UI_COMMAND(ConstrainChildBodiesToParentBody, "Constraint selected bodies", "Create a constraint between all selected bodies as children and the last selected body as parent", EUserInterfaceActionType::Button, FInputChord(EModifierKey::Control, EKeys::Y));
+	UI_COMMAND(ConstrainChildBodiesToParentBody, "Constrain selected bodies", "Create a constraint between all selected bodies as children and the last selected body as parent", EUserInterfaceActionType::Button, FInputChord(EModifierKey::Control, EKeys::Y));
 	UI_COMMAND(ResetConstraint, "Reset", "Reset Constraint", EUserInterfaceActionType::Button, FInputChord(EModifierKey::Alt, EKeys::P));
 	UI_COMMAND(SnapConstraint, "Snap", "Snap Constraint Translation To Bone", EUserInterfaceActionType::Button, FInputChord(EModifierKey::Control, EKeys::P));
 	UI_COMMAND(ConvertToBallAndSocket, "To Ball & Socket", "Convert Selected Constraint To Ball-And-Socket", EUserInterfaceActionType::Button, FInputChord());
