Physics Asset Profiles Reskin

diff --git a/Engine/Source/Editor/EditorStyle/Private/StarshipStyle.cpp b/Engine/Source/Editor/EditorStyle/Private/StarshipStyle.cpp
--- a/Engine/Source/Editor/EditorStyle/Private/StarshipStyle.cpp
+++ b/Engine/Source/Editor/EditorStyle/Private/StarshipStyle.cpp
@@ -460,6 +460,7 @@ void FStarshipEditorStyle::FStyle::SetupGeneralStyles()
 		Set("Icons.OpenSourceLocation", new IMAGE_BRUSH_SVG("Starship/Common/OpenSourceLocation", Icon16x16));
 		Set("Icons.Find", new IMAGE_BRUSH_SVG("Starship/Common/Find", Icon16x16));
 		Set("Icons.Validate", new IMAGE_BRUSH_SVG("Starship/Common/validate", Icon16x16));
+		Set("Icons.Tools", new IMAGE_BRUSH_SVG("Starship/Common/EditorModes", Icon16x16));
 
 
 		Set("Icons.Toolbar.Play", new IMAGE_BRUSH_SVG("Starship/Common/play", Icon20x20));
@@ -7191,7 +7192,7 @@ void FStarshipEditorStyle::FStyle::SetupToolkitStyles()
 	{
 		Set( "PhysicsAssetEditor.Tabs.Properties", new IMAGE_BRUSH( "/Icons/icon_tab_SelectionDetails_16x", Icon16x16 ) );
 		Set( "PhysicsAssetEditor.Tabs.Hierarchy", new IMAGE_BRUSH( "/Icons/levels_16x", Icon16x16 ) );
-		Set( "PhysicsAssetEditor.Tabs.Profiles", new IMAGE_BRUSH( "/PhysicsAssetEditor/icon_ProfilesTab_16x", Icon16x16 ) );
+		Set( "PhysicsAssetEditor.Tabs.Profiles", new IMAGE_BRUSH_SVG( "Starship/AssetEditors/ProfileFolder", Icon16x16 ) );
 		Set( "PhysicsAssetEditor.Tabs.Graph", new IMAGE_BRUSH( "/PhysicsAssetEditor/icon_GraphTab_16x", Icon16x16 ) );
 		Set( "PhysicsAssetEditor.Tabs.Tools", new IMAGE_BRUSH( "/PhysicsAssetEditor/icon_ToolsTab_16x", Icon16x16 ) );
 
@@ -7236,12 +7237,8 @@ void FStarshipEditorStyle::FStyle::SetupToolkitStyles()
 		Set( "PhysicsAssetEditor.ConvertToSkeletal", new IMAGE_BRUSH_SVG("Starship/Persona/AnimationToSkeletal", Icon20x20));
 		Set( "PhysicsAssetEditor.DeleteConstraint", new IMAGE_BRUSH( "PhysicsAssetEditor/icon_PhAT_DeleteConstraint_40x", Icon40x40 ) );
 
-		Set("PhysicsAssetEditor.NewPhysicalAnimationProfile", new IMAGE_BRUSH("PhysicsAssetEditor/icon_PhAT_NewBody_40x", Icon20x20));
-		Set("PhysicsAssetEditor.DeleteCurrentPhysicalAnimationProfile", new IMAGE_BRUSH("PhysicsAssetEditor/icon_PhAT_DeletePrimitive_40x", Icon20x20));
 		Set("PhysicsAssetEditor.AddBodyToPhysicalAnimationProfile", new IMAGE_BRUSH("PhysicsAssetEditor/icon_PhAT_NewBody_40x", Icon20x20));
 		Set("PhysicsAssetEditor.RemoveBodyFromPhysicalAnimationProfile", new IMAGE_BRUSH("PhysicsAssetEditor/icon_PhAT_DeletePrimitive_40x", Icon20x20));
-		Set("PhysicsAssetEditor.NewConstraintProfile", new IMAGE_BRUSH("PhysicsAssetEditor/icon_PHatMode_Joint_40x", Icon20x20));
-		Set("PhysicsAssetEditor.DeleteCurrentConstraintProfile", new IMAGE_BRUSH("PhysicsAssetEditor/icon_PhAT_DeleteConstraint_40x", Icon20x20));
 		Set("PhysicsAssetEditor.AddConstraintToCurrentConstraintProfile", new IMAGE_BRUSH("PhysicsAssetEditor/icon_PHatMode_Joint_40x", Icon20x20));
 		Set("PhysicsAssetEditor.RemoveConstraintFromCurrentConstraintProfile", new IMAGE_BRUSH("PhysicsAssetEditor/icon_PhAT_DeleteConstraint_40x", Icon20x20));
 
@@ -7256,6 +7253,10 @@ void FStarshipEditorStyle::FStyle::SetupToolkitStyles()
 		Set("PhysicsAssetEditor.Tree.TaperedCapsule", new IMAGE_BRUSH("PhysicsAssetEditor/TaperedCapsule_16x", Icon16x16));
 		Set("PhysicsAssetEditor.Tree.Constraint", new IMAGE_BRUSH("PhysicsAssetEditor/Constraint_16x", Icon16x16));
 
+		Set("PhysicsAssetEditor.BoneAssign", new IMAGE_BRUSH_SVG("Starship/Persona/BoneAssign", Icon20x20));
+		Set("PhysicsAssetEditor.BoneUnassign", new IMAGE_BRUSH_SVG("Starship/Persona/BoneUnassign", Icon20x20));
+		Set("PhysicsAssetEditor.BoneLocate", new IMAGE_BRUSH_SVG("Starship/Persona/BoneLocate", Icon20x20));
+
 		Set("PhysicsAssetEditor.Tree.Font", DEFAULT_FONT("Regular", 10));
 
 		Set("PhysicsAssetEditor.Graph.TextStyle", FTextBlockStyle(NormalText)
@@ -7269,12 +7270,13 @@ void FStarshipEditorStyle::FStyle::SetupToolkitStyles()
 		Set("PhysicsAssetEditor.Graph.Node.ShadowSelected", new BOX_BRUSH( "PhysicsAssetEditor/PhysicsNode_shadow_selected", FMargin(18.0f/64.0f) ) );
 		Set("PhysicsAssetEditor.Graph.Node.Shadow", new BOX_BRUSH( "Graph/RegularNode_shadow", FMargin(18.0f/64.0f) ) );
 
-		FEditableTextBoxStyle EditableTextBlock = FEditableTextBoxStyle()
+		FEditableTextBoxStyle EditableTextBlock = NormalEditableTextBoxStyle
 			.SetFont(NormalText.Font)
-			.SetBackgroundImageNormal(BOX_BRUSH("Common/TextBox", FMargin(4.0f / 16.0f)))
-			.SetBackgroundImageHovered(BOX_BRUSH("Common/TextBox_Hovered", FMargin(4.0f / 16.0f)))
-			.SetBackgroundImageFocused(BOX_BRUSH("Common/TextBox_Hovered", FMargin(4.0f / 16.0f)))
-			.SetBackgroundImageReadOnly(BOX_BRUSH("Common/TextBox_ReadOnly", FMargin(4.0f / 16.0f)));
+			.SetBackgroundImageNormal(FSlateNoResource())
+			.SetBackgroundImageHovered(FSlateNoResource())
+			.SetBackgroundImageFocused(FSlateNoResource())
+			.SetBackgroundImageReadOnly(FSlateNoResource())
+			.SetForegroundColor(FSlateColor::UseStyle());
 
 		Set("PhysicsAssetEditor.Profiles.EditableTextBoxStyle", EditableTextBlock);
 
diff --git a/Engine/Source/Editor/PhysicsAssetEditor/PhysicsAssetEditor.Build.cs b/Engine/Source/Editor/PhysicsAssetEditor/PhysicsAssetEditor.Build.cs
--- a/Engine/Source/Editor/PhysicsAssetEditor/PhysicsAssetEditor.Build.cs
+++ b/Engine/Source/Editor/PhysicsAssetEditor/PhysicsAssetEditor.Build.cs
@@ -47,6 +47,7 @@ public class PhysicsAssetEditor : ModuleRules
 				"PhysicsUtilities",
 				"MeshUtilitiesCommon",
 				"ApplicationCore",
+				"ToolWidgets"
 			}
         );
 
diff --git a/Engine/Source/Editor/PhysicsAssetEditor/Private/PhysicsAssetDetailsCustomization.cpp b/Engine/Source/Editor/PhysicsAssetEditor/Private/PhysicsAssetDetailsCustomization.cpp
--- a/Engine/Source/Editor/PhysicsAssetEditor/Private/PhysicsAssetDetailsCustomization.cpp
+++ b/Engine/Source/Editor/PhysicsAssetEditor/Private/PhysicsAssetDetailsCustomization.cpp
@@ -23,6 +23,7 @@
 #include "Widgets/Input/SEditableTextBox.h"
 #include "PropertyHandle.h"
 #include "PhysicsAssetEditorActions.h"
+#include "Styling/StyleColors.h"
 
 #define LOCTEXT_NAMESPACE "PhysicsAssetDetailsCustomization"
 
@@ -98,6 +99,12 @@ void FPhysicsAssetDetailsCustomization::BindCommands()
 		FCanExecuteAction::CreateSP(this, &FPhysicsAssetDetailsCustomization::CanRemoveBodyFromPhysicalAnimationProfile)
 	);
 
+	CommandList->MapAction(
+		Commands.SelectAllBodiesInCurrentPhysicalAnimationProfile,
+		FExecuteAction::CreateSP(this, &FPhysicsAssetDetailsCustomization::SelectAllBodiesInCurrentPhysicalAnimationProfile),
+		FCanExecuteAction::CreateSP(this, &FPhysicsAssetDetailsCustomization::CanSelectAllBodiesInCurrentPhysicalAnimationProfile)
+	);
+
 	CommandList->MapAction(
 		Commands.NewConstraintProfile,
 		FExecuteAction::CreateSP(this, &FPhysicsAssetDetailsCustomization::NewConstraintProfile),
@@ -127,6 +134,12 @@ void FPhysicsAssetDetailsCustomization::BindCommands()
 		FExecuteAction::CreateSP(this, &FPhysicsAssetDetailsCustomization::RemoveConstraintFromCurrentConstraintProfile),
 		FCanExecuteAction::CreateSP(this, &FPhysicsAssetDetailsCustomization::CanRemoveConstraintFromCurrentConstraintProfile)
 	);
+
+	CommandList->MapAction(
+		Commands.SelectAllBodiesInCurrentConstraintProfile,
+		FExecuteAction::CreateSP(this, &FPhysicsAssetDetailsCustomization::SelectAllBodiesInCurrentConstraintProfile),
+		FCanExecuteAction::CreateSP(this, &FPhysicsAssetDetailsCustomization::CanSelectAllBodiesInCurrentConstraintProfile)
+	);
 }
 
 TSharedRef< SWidget > FPhysicsAssetDetailsCustomization::FillPhysicalAnimationProfileOptions()
@@ -143,13 +156,36 @@ TSharedRef< SWidget > FPhysicsAssetDetailsCustomization::FillPhysicalAnimationPr
 
 	if(SharedData->PhysicsAsset)
 	{
+		MenuBuilder.BeginSection("NewProfile", LOCTEXT("PhysicsAssetEditor_NewPhysicalAnimationMenu", "New"));
+		{
+			MenuBuilder.AddMenuEntry(Commands.NewPhysicalAnimationProfile);
+		}
+		MenuBuilder.EndSection();
+
 		MenuBuilder.BeginSection("CurrentProfile", LOCTEXT("PhysicsAssetEditor_CurrentPhysicalAnimationMenu", "Current Profile"));
 		{
+			MenuBuilder.AddMenuEntry(
+				LOCTEXT("PhysicsAssetEditor_RenameMenu", "Rename"),
+				LOCTEXT("PhysicsAssetEditor_RenameMenuTooltip", "Rename the Current Physical Animation Profile"),
+				FSlateIcon(),
+				FUIAction(
+					FExecuteAction::CreateLambda([this]
+						{
+							bIsRenamePending = true;
+							FSlateApplication::Get().SetKeyboardFocus(PhysicalAnimationProfileNameTextBox);
+							FSlateApplication::Get().SetUserFocus(0, PhysicalAnimationProfileNameTextBox);
+						}),
+					FCanExecuteAction::CreateLambda([this]
+						{
+							return PhysicsAssetEditorPtr.Pin()->GetSharedData()->PhysicsAsset->CurrentPhysicalAnimationProfileName != NAME_None;
+						}))
+			);
 			MenuBuilder.AddMenuEntry(Commands.DuplicatePhysicalAnimationProfile);
+			MenuBuilder.AddMenuEntry(Commands.DeleteCurrentPhysicalAnimationProfile);
 		}
 		MenuBuilder.EndSection();
 
-		MenuBuilder.BeginSection("PhysicalAnimationProfile", LOCTEXT("PhysicsAssetEditor_PhysicalAnimationMenu", "Physical Animation Profiles"));
+		MenuBuilder.BeginSection("PhysicalAnimationProfile", LOCTEXT("PhysicsAssetEditor_PhysicalAnimationMenu", "Animation Profiles"));
 		{
 			TArray<FName> ProfileNames;
 			ProfileNames.Add(NAME_None);
@@ -185,58 +221,10 @@ TSharedRef< SWidget > FPhysicsAssetDetailsCustomization::FillPhysicalAnimationPr
 					return SharedData->PhysicsAsset->CurrentPhysicalAnimationProfileName == ProfileName ? ECheckBoxState::Checked : ECheckBoxState::Unchecked;
 				});
 
-				auto SearchClickedLambda = [ProfileName, SharedData]()
-				{
-					TArray<int32> NewBodiesSelection;
-					for (int32 BSIndex = 0; BSIndex < SharedData->PhysicsAsset->SkeletalBodySetups.Num(); ++BSIndex)
-					{
-						const USkeletalBodySetup* BS = SharedData->PhysicsAsset->SkeletalBodySetups[BSIndex];
-						if (BS->FindPhysicalAnimationProfile(ProfileName))
-						{
-							NewBodiesSelection.Add(BSIndex);
-						}
-					}
-					SharedData->ClearSelectedBody();	//clear selection
-					SharedData->SetSelectedBodiesAnyPrim(NewBodiesSelection, true);
-
-					FSlateApplication::Get().DismissAllMenus();
+				TSharedRef<SWidget> PhysAnimProfileButton = SNew(STextBlock)
+					.Text(FText::FromString(ProfileName.ToString()));
 
-					return FReply::Handled();
-				};
-					
-				TSharedRef<SWidget> PhysAnimProfileButton = SNew(SHorizontalBox)
-					+ SHorizontalBox::Slot()
-					.FillWidth(1.f)
-					.VAlign(VAlign_Center)
-					[
-						SNew(STextBlock)
-						.Text(FText::FromString(ProfileName.ToString()))
-					]
-					+ SHorizontalBox::Slot()
-					.AutoWidth()
-					.Padding(2.f, 0.f, 0.f, 0.f)
-					.VAlign(VAlign_Center)
-					[
-						SNew(SButton)
-						.ButtonStyle(FEditorStyle::Get(), "HoverHintOnly")
-						.ToolTipText(LOCTEXT("SelectBodies", "Select all bodies that are assigned to this profile."))
-						.OnClicked_Lambda(SearchClickedLambda)
-						[
-							SNew(SBox)
-							.WidthOverride(MenuIconSize)
-							.HeightOverride(MenuIconSize)
-							.Visibility_Lambda([ProfileName](){ return (ProfileName == NAME_None) ? EVisibility::Collapsed : EVisibility::Visible; })
-							[
-								SNew(SImage)
-								.Image(FSlateIcon(FEditorStyle::GetStyleSetName(), "Symbols.SearchGlass").GetIcon())
-							]
-								
-						]
-					];
-
-
-				//MenuBuilder.AddMenuEntry( FText::FromString(ProfileName.ToString()), FText::GetEmpty(), FSlateIcon(), Action, NAME_None, EUserInterfaceActionType::Check);
-				MenuBuilder.AddMenuEntry(Action, PhysAnimProfileButton, NAME_None, TAttribute<FText>(), EUserInterfaceActionType::Check);
+				MenuBuilder.AddMenuEntry(Action, PhysAnimProfileButton, NAME_None, TAttribute<FText>(), EUserInterfaceActionType::RadioButton);
 			}
 		}
 
@@ -261,9 +249,33 @@ TSharedRef< SWidget > FPhysicsAssetDetailsCustomization::FillConstraintProfilesO
 
 	if(SharedData->PhysicsAsset)
 	{
+		MenuBuilder.BeginSection("NewProfile", LOCTEXT("PhysicsAssetEditor_NewConstraintProfileMenu", "New"));
+		{
+			MenuBuilder.AddMenuEntry(Commands.NewConstraintProfile);
+		}
+		MenuBuilder.EndSection();
+
 		MenuBuilder.BeginSection("CurrentProfile", LOCTEXT("PhysicsAssetEditor_CurrentProfileMenu", "Current Profile"));
 		{
+			MenuBuilder.AddMenuEntry(
+				LOCTEXT("PhysicsAssetEditor_RenameMenu", "Rename"),
+				LOCTEXT("PhysicsAssetEditor_RenameMenuTooltip", "Rename the Current Constraint Profile"),
+				FSlateIcon(),
+				FUIAction(
+					FExecuteAction::CreateLambda([this]
+						{
+							bIsRenamePending = true;
+							FSlateApplication::Get().SetKeyboardFocus(ConstraintProfileNameTextBox);
+							FSlateApplication::Get().SetUserFocus(0, ConstraintProfileNameTextBox);
+						}),
+					FCanExecuteAction::CreateLambda([this]
+						{
+							return PhysicsAssetEditorPtr.Pin()->GetSharedData()->PhysicsAsset->CurrentConstraintProfileName != NAME_None;
+						}))
+			);
+
 			MenuBuilder.AddMenuEntry(Commands.DuplicateConstraintProfile);
+			MenuBuilder.AddMenuEntry(Commands.DeleteCurrentConstraintProfile);
 		}
 		MenuBuilder.EndSection();
 
@@ -302,56 +314,10 @@ TSharedRef< SWidget > FPhysicsAssetDetailsCustomization::FillConstraintProfilesO
 					return SharedData->PhysicsAsset->CurrentConstraintProfileName == ProfileName ? ECheckBoxState::Checked : ECheckBoxState::Unchecked;
 				});
 
-				auto SearchClickedLambda = [ProfileName, SharedData]()
-				{
-					TArray<int32> NewSelectedConstraints;
-					for (int32 CSIndex = 0; CSIndex < SharedData->PhysicsAsset->ConstraintSetup.Num(); ++CSIndex)
-					{
-						const UPhysicsConstraintTemplate* CS = SharedData->PhysicsAsset->ConstraintSetup[CSIndex];
-						if (CS->ContainsConstraintProfile(ProfileName))
-						{
-							NewSelectedConstraints.AddUnique(CSIndex);
-						}
-					}
-
-					SharedData->ClearSelectedConstraints();	//clear selection
-					SharedData->SetSelectedConstraints(NewSelectedConstraints, true);
-
-					FSlateApplication::Get().DismissAllMenus();
-
-					return FReply::Handled();
-				};
-
-				TSharedRef<SWidget> ConstraintProfileButton = SNew(SHorizontalBox)
-				+ SHorizontalBox::Slot()
-				.FillWidth(1.f)
-				.VAlign(VAlign_Center)
-				[
-					SNew(STextBlock)
-					.Text(FText::FromString(ProfileName.ToString()))
-				]
-				+ SHorizontalBox::Slot()
-				.AutoWidth()
-				.Padding(2.f, 0.f, 0.f, 0.f)
-				.VAlign(VAlign_Center)
-				[
-					SNew(SButton)
-					.ButtonStyle(FEditorStyle::Get(), "HoverHintOnly")
-					.ToolTipText(LOCTEXT("SelectConstraints", "Select all constraints that are assigned to this profile."))
-					.OnClicked_Lambda(SearchClickedLambda)
-					[
-						SNew(SBox)
-						.WidthOverride(MenuIconSize)
-						.HeightOverride(MenuIconSize)
-						.Visibility_Lambda([ProfileName]() { return (ProfileName == NAME_None) ? EVisibility::Collapsed : EVisibility::Visible; })
-						[
-							SNew(SImage)
-							.Image(FSlateIcon(FEditorStyle::GetStyleSetName(), "Symbols.SearchGlass").GetIcon())
-						]
-					]
-				];
-
-				MenuBuilder.AddMenuEntry(Action, ConstraintProfileButton, NAME_None, TAttribute<FText>(), EUserInterfaceActionType::Check);
+				TSharedRef<SWidget> ConstraintProfileButton = SNew(STextBlock)
+					.Text(FText::FromString(ProfileName.ToString()));
+					
+				MenuBuilder.AddMenuEntry(Action, ConstraintProfileButton, NAME_None, TAttribute<FText>(), EUserInterfaceActionType::RadioButton);
 			}
 		}
 		MenuBuilder.EndSection();
@@ -363,6 +329,7 @@ TSharedRef< SWidget > FPhysicsAssetDetailsCustomization::FillConstraintProfilesO
 void FPhysicsAssetDetailsCustomization::HandlePhysicalAnimationProfileNameCommitted(const FText& InText, ETextCommit::Type InCommitType)
 {
 	PhysicalAnimationProfileNameTextBox->SetError(FText::GetEmpty());
+	bIsRenamePending = false;
 
 	if(InCommitType != ETextCommit::OnCleared)
 	{
@@ -392,6 +359,7 @@ void FPhysicsAssetDetailsCustomization::HandlePhysicalAnimationProfileNameCommit
 void FPhysicsAssetDetailsCustomization::HandleConstraintProfileNameCommitted(const FText& InText, ETextCommit::Type InCommitType)
 {
 	ConstraintProfileNameTextBox->SetError(FText::GetEmpty());
+	bIsRenamePending = false;
 
 	if(InCommitType != ETextCommit::OnCleared)
 	{
@@ -418,47 +386,30 @@ void FPhysicsAssetDetailsCustomization::HandleConstraintProfileNameCommitted(con
 	}
 }
 
-TSharedRef<SWidget> FPhysicsAssetDetailsCustomization::CreateProfileButton(const FText& InGlyph, TSharedPtr<FUICommandInfo> InCommand)
+TSharedRef<SWidget> FPhysicsAssetDetailsCustomization::CreateProfileButton(const FName& InIconName, TSharedPtr<FUICommandInfo> InCommand)
 {
 	check(InCommand.IsValid());
 
 	TWeakPtr<FUICommandInfo> LocalCommandPtr = InCommand;
 
 	return SNew(SButton)
-		.VAlign(EVerticalAlignment::VAlign_Center)
-		.ButtonStyle(FEditorStyle::Get(), "FlatButton")
-		.ForegroundColor(FEditorStyle::GetSlateColor("DefaultForeground"))
+		.ButtonStyle(FAppStyle::Get(), "SimpleButton")
 		.ToolTipText(InCommand->GetDescription())
 		.IsEnabled_Lambda([this, LocalCommandPtr]()
-		{
-			TSharedPtr<FUICommandList> CommandList = PhysicsAssetEditorPtr.Pin()->GetToolkitCommands();
-			return CommandList->CanExecuteAction(LocalCommandPtr.Pin().ToSharedRef());
-		})
+			{
+				TSharedPtr<FUICommandList> CommandList = PhysicsAssetEditorPtr.Pin()->GetToolkitCommands();
+				return CommandList->CanExecuteAction(LocalCommandPtr.Pin().ToSharedRef());
+			})
 		.OnClicked(FOnClicked::CreateLambda([this, LocalCommandPtr]()
-		{
-			TSharedPtr<FUICommandList> CommandList = PhysicsAssetEditorPtr.Pin()->GetToolkitCommands();
-			return CommandList->ExecuteAction(LocalCommandPtr.Pin().ToSharedRef()) ? FReply::Handled() : FReply::Unhandled();
-		}))
+			{
+				TSharedPtr<FUICommandList> CommandList = PhysicsAssetEditorPtr.Pin()->GetToolkitCommands();
+				return CommandList->ExecuteAction(LocalCommandPtr.Pin().ToSharedRef()) ? FReply::Handled() : FReply::Unhandled();
+			}))
+		.ContentPadding(0)
 		[
-			SNew( SHorizontalBox )
-			+SHorizontalBox::Slot()
-			.VAlign(VAlign_Center)
-			.AutoWidth()
-			[
-				SNew(STextBlock)
-				.TextStyle(FEditorStyle::Get(), "PhysicsAssetEditor.Profiles.Font")
-				.Font(FEditorStyle::Get().GetFontStyle("FontAwesome.11"))
-				.Text(InGlyph)
-			]
-			+SHorizontalBox::Slot()
-			.AutoWidth()
-			.VAlign(VAlign_Center)
-			.Padding(4, 0, 0, 0)
-			[
-				SNew(STextBlock)
-				.TextStyle( FEditorStyle::Get(), "PhysicsAssetEditor.Profiles.Font" )
-				.Text(InCommand->GetLabel())
-			]
+			SNew(SImage)
+			.Image(FAppStyle::Get().GetBrush(InIconName))
+			.ColorAndOpacity(FSlateColor::UseForeground())
 		];
 }
 
@@ -470,39 +421,49 @@ TSharedRef<SWidget> FPhysicsAssetDetailsCustomization::MakePhysicalAnimationProf
 
 	return SNew(SHorizontalBox)
 		.ToolTipText(LOCTEXT("CurrentPhysicalAnimationProfileWidgetTooltip", "Select and edit the current physical animation profile."))
-		+SHorizontalBox::Slot()
-		.FillWidth(1.0f)
+
+		+ SHorizontalBox::Slot()
+		.AutoWidth()
 		.HAlign(HAlign_Left)
 		.VAlign(VAlign_Center)
+		.Padding(0.0f, 0.0f, 2.0f, 3.0f)
+		[
+			SNew(STextBlock)
+			.Text(LOCTEXT("CurrentProfile", "Current Profile"))
+		]
+
+		+SHorizontalBox::Slot()
+		.HAlign(HAlign_Right)
+		.VAlign(VAlign_Center)
 		[
-			SNew(SComboButton)
-			.OnGetMenuContent(this, &FPhysicsAssetDetailsCustomization::FillPhysicalAnimationProfileOptions)
-			.ButtonContent()
+			SNew(SVerticalBox)
+			
+			+SVerticalBox::Slot()
+			.AutoHeight()
+			.Padding(0, 6, 8, 0)
 			[
-				SNew(SVerticalBox)
-				+SVerticalBox::Slot()
-				.AutoHeight()
-				.Padding(0.0f, 0.0f, 2.0f, 3.0f)
+				SNew(SBox)
+				.WidthOverride(125.0f)
 				[
-					SNew(STextBlock)
-					.Text(LOCTEXT("CurrentProfile", "Current Profile"))
-				]
-				+SVerticalBox::Slot()
-				.AutoHeight()
-				[
-					SNew(SHorizontalBox)
-					+SHorizontalBox::Slot()
-					.FillWidth(1.0f)
-					.Padding(0.0f, 0.0f, 2.0f, 0.0f)
+					SNew(SComboButton)
+					.OnGetMenuContent(this, &FPhysicsAssetDetailsCustomization::FillPhysicalAnimationProfileOptions)
+					.ButtonContent()
 					[
 						SAssignNew(PhysicalAnimationProfileNameTextBox, SEditableTextBox)
+						.Style(FAppStyle::Get(), "PhysicsAssetEditor.Profiles.EditableTextBoxStyle")
 						.Text_Lambda([LocalPhysicsAssetEditorPtr]()
 						{
 							return FText::FromName(LocalPhysicsAssetEditorPtr.Pin()->GetSharedData()->PhysicsAsset->CurrentPhysicalAnimationProfileName);
 						})
-						.IsEnabled_Lambda([LocalPhysicsAssetEditorPtr]()
+						.ForegroundColor_Lambda([LocalPhysicsAssetEditorPtr]() -> FSlateColor
+						{
+							FName ProfileName = LocalPhysicsAssetEditorPtr.Pin()->GetSharedData()->PhysicsAsset->CurrentPhysicalAnimationProfileName;
+
+							return ProfileName == NAME_None ? FStyleColors::Foreground : FStyleColors::White;
+						})
+						.IsEnabled_Lambda([this]()
 						{
-							return LocalPhysicsAssetEditorPtr.Pin()->GetSharedData()->PhysicsAsset->CurrentPhysicalAnimationProfileName != NAME_None;
+							return bIsRenamePending;
 						})
 						.OnTextChanged_Lambda([this, LocalPhysicsAssetEditorPtr](const FText& InText)
 						{
@@ -521,27 +482,32 @@ TSharedRef<SWidget> FPhysicsAssetDetailsCustomization::MakePhysicalAnimationProf
 					]
 				]
 			]
-		]
-		+SHorizontalBox::Slot()
-		.AutoWidth()
-		[
-			SNew(SUniformGridPanel)
-			.SlotPadding(FMargin(1.0f, 1.0f))
-			+SUniformGridPanel::Slot(0, 0)
-			[
-				CreateProfileButton(FEditorFontGlyphs::File, Commands.NewPhysicalAnimationProfile)
-			]
-			+SUniformGridPanel::Slot(1, 0)
-			[
-				CreateProfileButton(FEditorFontGlyphs::Trash, Commands.DeleteCurrentPhysicalAnimationProfile)
-			]
-			+SUniformGridPanel::Slot(0, 1)
-			[
-				CreateProfileButton(FEditorFontGlyphs::Plus_Circle, Commands.AddBodyToPhysicalAnimationProfile)
-			]
-			+SUniformGridPanel::Slot(1, 1)
+			+SVerticalBox::Slot()
+			.AutoHeight()
+			.Padding(0, 2, 0, 7)
 			[
-				CreateProfileButton(FEditorFontGlyphs::Minus_Circle, Commands.RemoveBodyFromPhysicalAnimationProfile)
+				SNew(SHorizontalBox)
+				+ SHorizontalBox::Slot()
+				.HAlign(HAlign_Left)
+				.Padding(0)
+				.AutoWidth()
+				[
+					CreateProfileButton("PhysicsAssetEditor.BoneAssign", Commands.AddBodyToPhysicalAnimationProfile)
+				]
+				+ SHorizontalBox::Slot()
+				.HAlign(HAlign_Left)
+				.Padding(0)
+				.AutoWidth()
+				[
+					CreateProfileButton("PhysicsAssetEditor.BoneUnassign", Commands.RemoveBodyFromPhysicalAnimationProfile)
+				]
+				+ SHorizontalBox::Slot()
+				.HAlign(HAlign_Left)
+				.Padding(0)
+				.AutoWidth()
+				[
+					CreateProfileButton("PhysicsAssetEditor.BoneLocate", Commands.SelectAllBodiesInCurrentPhysicalAnimationProfile)
+				]
 			]
 		];
 }
@@ -554,39 +520,49 @@ TSharedRef<SWidget> FPhysicsAssetDetailsCustomization::MakeConstraintProfilesWid
 
 	return SNew(SHorizontalBox)
 		.ToolTipText(LOCTEXT("CurrentConstraintProfileWidgetTooltip", "Select and edit the current constraint profile."))
-		+SHorizontalBox::Slot()
-		.FillWidth(1.0f)
+
+		+ SHorizontalBox::Slot()
+		.AutoWidth()
 		.HAlign(HAlign_Left)
 		.VAlign(VAlign_Center)
+		.Padding(0.0f, 0.0f, 2.0f, 3.0f)
+		[
+			SNew(STextBlock)
+			.Text(LOCTEXT("CurrentProfile", "Current Profile"))
+		]
+
+		+SHorizontalBox::Slot()
+		.HAlign(HAlign_Right)
+		.VAlign(VAlign_Center)
 		[
-			SNew(SComboButton)
-			.OnGetMenuContent(this, &FPhysicsAssetDetailsCustomization::FillConstraintProfilesOptions)
-			.ButtonContent()
+			SNew(SVerticalBox)
+			
+			+SVerticalBox::Slot()
+			.AutoHeight()
+			.Padding(0, 6, 8, 0)
 			[
-				SNew(SVerticalBox)
-				+SVerticalBox::Slot()
-				.AutoHeight()
-				.Padding(0.0f, 0.0f, 2.0f, 3.0f)
-				[
-					SNew(STextBlock)
-					.Text(LOCTEXT("CurrentProfile", "Current Profile"))
-				]
-				+SVerticalBox::Slot()
-				.AutoHeight()
+				SNew(SBox)
+				.WidthOverride(125.0f)
 				[
-					SNew(SHorizontalBox)
-					+SHorizontalBox::Slot()
-					.FillWidth(1.0f)
-					.Padding(0.0f, 0.0f, 2.0f, 0.0f)
+					SNew(SComboButton)
+					.OnGetMenuContent(this, &FPhysicsAssetDetailsCustomization::FillConstraintProfilesOptions)
+					.ButtonContent()
 					[
 						SAssignNew(ConstraintProfileNameTextBox, SEditableTextBox)
+						.Style(FAppStyle::Get(), "PhysicsAssetEditor.Profiles.EditableTextBoxStyle")
+						.ForegroundColor_Lambda([LocalPhysicsAssetEditorPtr]() -> FSlateColor
+						{
+							FName ProfileName = LocalPhysicsAssetEditorPtr.Pin()->GetSharedData()->PhysicsAsset->CurrentConstraintProfileName;
+
+							return ProfileName == NAME_None ? FStyleColors::Foreground : FStyleColors::White;
+						})
 						.Text_Lambda([LocalPhysicsAssetEditorPtr]()
 						{
 							return FText::FromName(LocalPhysicsAssetEditorPtr.Pin()->GetSharedData()->PhysicsAsset->CurrentConstraintProfileName);
 						})
-						.IsEnabled_Lambda([LocalPhysicsAssetEditorPtr]()
+						.IsEnabled_Lambda([this]()
 						{
-							return LocalPhysicsAssetEditorPtr.Pin()->GetSharedData()->PhysicsAsset->CurrentConstraintProfileName != NAME_None;
+							return bIsRenamePending;
 						})
 						.OnTextChanged_Lambda([this, LocalPhysicsAssetEditorPtr](const FText& InText)
 						{
@@ -605,27 +581,32 @@ TSharedRef<SWidget> FPhysicsAssetDetailsCustomization::MakeConstraintProfilesWid
 					]
 				]
 			]
-		]
-		+SHorizontalBox::Slot()
-		.AutoWidth()
-		[
-			SNew(SUniformGridPanel)
-			.SlotPadding(FMargin(1.0f, 1.0f))
-			+SUniformGridPanel::Slot(0, 0)
-			[
-				CreateProfileButton(FEditorFontGlyphs::File, Commands.NewConstraintProfile)
-			]
-			+SUniformGridPanel::Slot(1, 0)
-			[
-				CreateProfileButton(FEditorFontGlyphs::Trash, Commands.DeleteCurrentConstraintProfile)
-			]
-			+SUniformGridPanel::Slot(0, 1)
-			[
-				CreateProfileButton(FEditorFontGlyphs::Plus_Circle, Commands.AddConstraintToCurrentConstraintProfile)
-			]
-			+SUniformGridPanel::Slot(1, 1)
+			+SVerticalBox::Slot()
+			.AutoHeight()
+			.Padding(0, 2, 0, 7)
 			[
-				CreateProfileButton(FEditorFontGlyphs::Minus_Circle, Commands.RemoveConstraintFromCurrentConstraintProfile)
+				SNew(SHorizontalBox)
+				+ SHorizontalBox::Slot()
+				.HAlign(HAlign_Left)
+				.Padding(0)
+				.AutoWidth()
+				[
+					CreateProfileButton("PhysicsAssetEditor.BoneAssign", Commands.AddConstraintToCurrentConstraintProfile)
+				]
+				+ SHorizontalBox::Slot()
+				.HAlign(HAlign_Left)
+				.Padding(0)
+				.AutoWidth()
+				[
+					CreateProfileButton("PhysicsAssetEditor.BoneUnassign", Commands.RemoveConstraintFromCurrentConstraintProfile)
+				]
+				+ SHorizontalBox::Slot()
+				.HAlign(HAlign_Left)
+				.Padding(0)
+				.AutoWidth()
+				[
+					CreateProfileButton("PhysicsAssetEditor.BoneLocate", Commands.SelectAllBodiesInCurrentConstraintProfile)
+				]
 			]
 		];
 }
@@ -803,6 +784,33 @@ bool FPhysicsAssetDetailsCustomization::CanRemoveBodyFromPhysicalAnimationProfil
 	return (PhysicsAssetEditorPtr.Pin()->IsNotSimulation() && bSelectedBodies && PhysicalAnimationProfileExistsForAny() && PhysicsAsset->CurrentPhysicalAnimationProfileName != NAME_None);
 }
 
+void FPhysicsAssetDetailsCustomization::SelectAllBodiesInCurrentPhysicalAnimationProfile()
+{
+	TArray<int32> NewBodiesSelection;
+	TSharedPtr<FPhysicsAssetEditorSharedData> SharedData = PhysicsAssetEditorPtr.Pin()->GetSharedData();
+
+	for (int32 BSIndex = 0; BSIndex < SharedData->PhysicsAsset->SkeletalBodySetups.Num(); ++BSIndex)
+	{
+		const USkeletalBodySetup* BS = SharedData->PhysicsAsset->SkeletalBodySetups[BSIndex];
+		FName ProfileName = BS->GetCurrentPhysicalAnimationProfileName();
+
+		if (BS->FindPhysicalAnimationProfile(ProfileName))
+		{
+			NewBodiesSelection.Add(BSIndex);
+		}
+	}
+	SharedData->ClearSelectedBody();	//clear selection
+	SharedData->SetSelectedBodiesAnyPrim(NewBodiesSelection, true);
+
+	return;
+}
+
+bool FPhysicsAssetDetailsCustomization::CanSelectAllBodiesInCurrentPhysicalAnimationProfile() const
+{
+	UPhysicsAsset* PhysicsAsset = PhysicsAssetEditorPtr.Pin()->GetSharedData()->PhysicsAsset;
+	return PhysicsAsset->CurrentPhysicalAnimationProfileName != NAME_None;
+}
+
 void FPhysicsAssetDetailsCustomization::ApplyConstraintProfile(FName InName)
 {
 	TSharedPtr<FPhysicsAssetEditorSharedData> SharedData = PhysicsAssetEditorPtr.Pin()->GetSharedData();
@@ -984,4 +992,32 @@ bool FPhysicsAssetDetailsCustomization::CanRemoveConstraintFromCurrentConstraint
 	return (PhysicsAssetEditorPtr.Pin()->IsNotSimulation() && bSelectedConstraints && PhysicsAsset->CurrentConstraintProfileName != NAME_None && ConstraintProfileExistsForAny());
 }
 
+void FPhysicsAssetDetailsCustomization::SelectAllBodiesInCurrentConstraintProfile()
+{
+	TArray<int32> NewSelectedConstraints;
+	TSharedPtr<FPhysicsAssetEditorSharedData> SharedData = PhysicsAssetEditorPtr.Pin()->GetSharedData();
+
+	for (int32 CSIndex = 0; CSIndex < SharedData->PhysicsAsset->ConstraintSetup.Num(); ++CSIndex)
+	{
+		const UPhysicsConstraintTemplate* CS = SharedData->PhysicsAsset->ConstraintSetup[CSIndex];
+		FName ProfileName = CS->GetCurrentConstraintProfileName();
+
+		if (CS->ContainsConstraintProfile(ProfileName))
+		{
+			NewSelectedConstraints.AddUnique(CSIndex);
+		}
+	}
+
+	SharedData->ClearSelectedConstraints();	//clear selection
+	SharedData->SetSelectedConstraints(NewSelectedConstraints, true);
+
+	return;
+}
+
+bool FPhysicsAssetDetailsCustomization::CanSelectAllBodiesInCurrentConstraintProfile() const
+{
+	UPhysicsAsset* PhysicsAsset = PhysicsAssetEditorPtr.Pin()->GetSharedData()->PhysicsAsset;
+	return PhysicsAsset->CurrentConstraintProfileName != NAME_None;
+}
+
 #undef LOCTEXT_NAMESPACE
diff --git a/Engine/Source/Editor/PhysicsAssetEditor/Private/PhysicsAssetDetailsCustomization.h b/Engine/Source/Editor/PhysicsAssetEditor/Private/PhysicsAssetDetailsCustomization.h
--- a/Engine/Source/Editor/PhysicsAssetEditor/Private/PhysicsAssetDetailsCustomization.h
+++ b/Engine/Source/Editor/PhysicsAssetEditor/Private/PhysicsAssetDetailsCustomization.h
@@ -21,7 +21,8 @@ public:
 	static TSharedRef<IDetailCustomization> MakeInstance(TWeakPtr<FPhysicsAssetEditor> InPhysicsAssetEditor);
 
 	FPhysicsAssetDetailsCustomization(TWeakPtr<FPhysicsAssetEditor> InPhysicsAssetEditor)
-		: PhysicsAssetEditorPtr(InPhysicsAssetEditor)
+		: PhysicsAssetEditorPtr(InPhysicsAssetEditor),
+		  bIsRenamePending(false)
 	{}
 
 	// IDetailCustomization interface
@@ -38,7 +39,7 @@ private:
 	
 	TSharedRef<SWidget> MakeConstraintProfilesWidget();
 
-	TSharedRef<SWidget> CreateProfileButton(const FText& InGlyph, TSharedPtr<FUICommandInfo> InCommand);
+	TSharedRef<SWidget> CreateProfileButton(const FName& InIconName, TSharedPtr<FUICommandInfo> InCommand);
 
 	void HandlePhysicalAnimationProfileNameCommitted(const FText& InText, ETextCommit::Type InCommitType);
 
@@ -66,6 +67,10 @@ private:
 
 	bool CanRemoveBodyFromPhysicalAnimationProfile() const;
 
+	void SelectAllBodiesInCurrentPhysicalAnimationProfile();
+
+	bool CanSelectAllBodiesInCurrentPhysicalAnimationProfile() const;
+
 	void ApplyConstraintProfile(FName InName);
 
 	bool ConstraintProfileExistsForAny() const;
@@ -90,6 +95,10 @@ private:
 
 	bool CanRemoveConstraintFromCurrentConstraintProfile() const;
 
+	void SelectAllBodiesInCurrentConstraintProfile();
+
+	bool CanSelectAllBodiesInCurrentConstraintProfile() const;
+
 private:
 	TWeakPtr<FPhysicsAssetEditor> PhysicsAssetEditorPtr;
 
@@ -102,4 +111,6 @@ private:
 	TSharedPtr<SEditableTextBox> PhysicalAnimationProfileNameTextBox;
 
 	TSharedPtr<SEditableTextBox> ConstraintProfileNameTextBox;
+
+	bool bIsRenamePending;
 };
diff --git a/Engine/Source/Editor/PhysicsAssetEditor/Private/PhysicsAssetEditorActions.cpp b/Engine/Source/Editor/PhysicsAssetEditor/Private/PhysicsAssetEditorActions.cpp
--- a/Engine/Source/Editor/PhysicsAssetEditor/Private/PhysicsAssetEditorActions.cpp
+++ b/Engine/Source/Editor/PhysicsAssetEditor/Private/PhysicsAssetEditorActions.cpp
@@ -101,13 +101,15 @@ void FPhysicsAssetEditorCommands::RegisterCommands()
 	UI_COMMAND(DuplicatePhysicalAnimationProfile, "Duplicate", "Duplicates the current physical animation profile", EUserInterfaceActionType::Button, FInputChord());
 	UI_COMMAND(DeleteCurrentPhysicalAnimationProfile, "Delete", "Deletes the current physical animation profile", EUserInterfaceActionType::Button, FInputChord());
 	UI_COMMAND(AddBodyToPhysicalAnimationProfile, "Assign", "Assigns the selected bodies to the current physical animation profile", EUserInterfaceActionType::Button, FInputChord());
+	UI_COMMAND(SelectAllBodiesInCurrentPhysicalAnimationProfile, "Select All Bodies", "Select all bodies that are assigned to the current physical animation profile", EUserInterfaceActionType::Button, FInputChord());
 	UI_COMMAND(RemoveBodyFromPhysicalAnimationProfile, "Unassign", "Unassigns the selected bodies from the current physical animation profile", EUserInterfaceActionType::Button, FInputChord());
 	UI_COMMAND(NewConstraintProfile, "New", "Creates a new constraint profile", EUserInterfaceActionType::Button, FInputChord());
 	UI_COMMAND(DuplicateConstraintProfile, "Duplicate", "Duplicates the current constraint profile", EUserInterfaceActionType::Button, FInputChord());
 	UI_COMMAND(DeleteCurrentConstraintProfile, "Delete", "Deletes the current constraint profile", EUserInterfaceActionType::Button, FInputChord());
 	UI_COMMAND(AddConstraintToCurrentConstraintProfile, "Assign", "Assigns the selected constraints to the current constraint profile", EUserInterfaceActionType::Button, FInputChord());
 	UI_COMMAND(RemoveConstraintFromCurrentConstraintProfile, "Unassign", "Unassigns the selected constraints from the current constraint profile", EUserInterfaceActionType::Button, FInputChord());
-
+	UI_COMMAND(SelectAllBodiesInCurrentConstraintProfile, "Select All Bodies", "Select all bodies that are assigned to the current constraint profile", EUserInterfaceActionType::Button, FInputChord());
+	
 	UI_COMMAND(ShowBodies, "Show Bodies", "Display bodies in the tree view. Bodies are a collection of primitive shapes used for collision.", EUserInterfaceActionType::ToggleButton, FInputChord());
 	UI_COMMAND(ShowSimulatedBodies, "Show Simulated Bodies", "Display simulated bodies in the tree view. Bodies are a collection of primitive shapes used for collision.", EUserInterfaceActionType::ToggleButton, FInputChord());
 	UI_COMMAND(ShowKinematicBodies, "Show Kinematic Bodies", "Display kinematic bodies in the tree view. Bodies are a collection of primitive shapes used for collision.", EUserInterfaceActionType::ToggleButton, FInputChord());
diff --git a/Engine/Source/Editor/PhysicsAssetEditor/Private/PhysicsAssetEditorActions.h b/Engine/Source/Editor/PhysicsAssetEditor/Private/PhysicsAssetEditorActions.h
--- a/Engine/Source/Editor/PhysicsAssetEditor/Private/PhysicsAssetEditorActions.h
+++ b/Engine/Source/Editor/PhysicsAssetEditor/Private/PhysicsAssetEditorActions.h
@@ -113,11 +113,13 @@ public:
 	TSharedPtr<FUICommandInfo> DeleteCurrentPhysicalAnimationProfile;
 	TSharedPtr<FUICommandInfo> AddBodyToPhysicalAnimationProfile;
 	TSharedPtr<FUICommandInfo> RemoveBodyFromPhysicalAnimationProfile;
+	TSharedPtr<FUICommandInfo> SelectAllBodiesInCurrentPhysicalAnimationProfile;
 	TSharedPtr<FUICommandInfo> NewConstraintProfile;
 	TSharedPtr<FUICommandInfo> DuplicateConstraintProfile;
 	TSharedPtr<FUICommandInfo> DeleteCurrentConstraintProfile;
 	TSharedPtr<FUICommandInfo> AddConstraintToCurrentConstraintProfile;
 	TSharedPtr<FUICommandInfo> RemoveConstraintFromCurrentConstraintProfile;
+	TSharedPtr<FUICommandInfo> SelectAllBodiesInCurrentConstraintProfile;
 	TSharedPtr<FUICommandInfo> ShowBodies;
 	TSharedPtr<FUICommandInfo> ShowSimulatedBodies;
 	TSharedPtr<FUICommandInfo> ShowKinematicBodies;
diff --git a/Engine/Source/Editor/PhysicsAssetEditor/Private/PhysicsAssetEditorProfilesSummoner.cpp b/Engine/Source/Editor/PhysicsAssetEditor/Private/PhysicsAssetEditorProfilesSummoner.cpp
--- a/Engine/Source/Editor/PhysicsAssetEditor/Private/PhysicsAssetEditorProfilesSummoner.cpp
+++ b/Engine/Source/Editor/PhysicsAssetEditor/Private/PhysicsAssetEditorProfilesSummoner.cpp
@@ -15,7 +15,7 @@ FPhysicsAssetEditorProfilesSummoner::FPhysicsAssetEditorProfilesSummoner(TShared
 	, PhysicsAssetPtr(InPhysicsAsset)
 {
 	TabLabel = LOCTEXT("PhysicsAssetProfilesTabTitle", "Profiles");
-	TabIcon = FSlateIcon(FEditorStyle::GetStyleSetName(), "PhysicsAssetEditor.Tabs.Profiles");
+	TabIcon = FSlateIcon(FAppStyle::GetAppStyleSetName(), "PhysicsAssetEditor.Tabs.Profiles");
 
 	bIsSingleton = true;
 
diff --git a/Engine/Source/Editor/PhysicsAssetEditor/Private/PhysicsAssetEditorSharedData.cpp b/Engine/Source/Editor/PhysicsAssetEditor/Private/PhysicsAssetEditorSharedData.cpp
--- a/Engine/Source/Editor/PhysicsAssetEditor/Private/PhysicsAssetEditorSharedData.cpp
+++ b/Engine/Source/Editor/PhysicsAssetEditor/Private/PhysicsAssetEditorSharedData.cpp
@@ -38,7 +38,7 @@
 #include "Exporters/Exporter.h"
 #include "Factories.h"
 #include "HAL/PlatformApplicationMisc.h"
-
+#include "SPrimaryButton.h"
 
 #define LOCTEXT_NAMESPACE "PhysicsAssetEditorShared"
 
@@ -2688,7 +2688,7 @@ TSharedRef<SWidget> FPhysicsAssetEditorSharedData::CreateGenerateBodiesWidget(co
 		.AutoHeight()
 		[
 			SNew(SBorder)
-			.BorderImage(FEditorStyle::GetBrush("ToolPanel.GroupBorder"))
+			.BorderImage(FAppStyle::Get().GetBrush("Brushes.Panel"))
 			.VAlign(VAlign_Center)
 			.HAlign(HAlign_Right)
 			[
@@ -2697,10 +2697,8 @@ TSharedRef<SWidget> FPhysicsAssetEditorSharedData::CreateGenerateBodiesWidget(co
 				.Padding(2.0f)
 				.AutoWidth()
 				[
-					SNew(SButton)
-					.ButtonStyle(FEditorStyle::Get(), "FlatButton.Success")
-					.ForegroundColor(FLinearColor::White)
-					.ContentPadding(FMargin(6, 2))
+					SNew(SPrimaryButton)
+					.Text(InCreateButtonText)
 					.OnClicked_Lambda([InOnCreate]()
 					{ 
 						GetMutableDefault<UPhysicsAssetGenerationSettings>()->SaveConfig(); 
@@ -2710,11 +2708,6 @@ TSharedRef<SWidget> FPhysicsAssetEditorSharedData::CreateGenerateBodiesWidget(co
 					.ToolTipText(bForNewAsset ? 
 								LOCTEXT("CreateAsset_Tooltip", "Create a new physics asset using these settings.") :
 								LOCTEXT("GenerateBodies_Tooltip", "Generate new bodies and constraints. If bodies are selected then they will be replaced along with their constraints using the new settings, otherwise all bodies and constraints will be re-created"))
-					[
-						SNew(STextBlock)
-						.TextStyle(FEditorStyle::Get(), "PhysicsAssetEditor.Tools.Font")
-						.Text(InCreateButtonText)
-					]
 				]
 				+SHorizontalBox::Slot()
 				.Padding(2.0f)
@@ -2722,7 +2715,7 @@ TSharedRef<SWidget> FPhysicsAssetEditorSharedData::CreateGenerateBodiesWidget(co
 				[
 					SNew(SButton)
 					.Visibility_Lambda([bForNewAsset](){ return bForNewAsset ? EVisibility::Visible : EVisibility::Collapsed; })
-					.ButtonStyle(FEditorStyle::Get(), "FlatButton")
+					.ButtonStyle(FAppStyle::Get(), "FlatButton")
 					.ForegroundColor(FLinearColor::White)
 					.ContentPadding(FMargin(6, 2))
 					.OnClicked_Lambda([InOnCancel](){ InOnCancel.ExecuteIfBound(); return FReply::Handled(); })
diff --git a/Engine/Source/Editor/PhysicsAssetEditor/Private/PhysicsAssetEditorToolsSummoner.cpp b/Engine/Source/Editor/PhysicsAssetEditor/Private/PhysicsAssetEditorToolsSummoner.cpp
--- a/Engine/Source/Editor/PhysicsAssetEditor/Private/PhysicsAssetEditorToolsSummoner.cpp
+++ b/Engine/Source/Editor/PhysicsAssetEditor/Private/PhysicsAssetEditorToolsSummoner.cpp
@@ -17,7 +17,7 @@ FPhysicsAssetEditorToolsSummoner::FPhysicsAssetEditorToolsSummoner(TSharedPtr<FA
 	: FWorkflowTabFactory("PhysicsAssetTools", InHostingApp)
 {
 	TabLabel = LOCTEXT("PhysicsAssetToolsTabTitle", "Tools");
-	TabIcon = FSlateIcon(FEditorStyle::GetStyleSetName(), "PhysicsAssetEditor.Tabs.Tools");
+	TabIcon = FSlateIcon(FAppStyle::GetAppStyleSetName(), "Icons.Tools");
 
 	bIsSingleton = true;
 
