Physics Asset Customization: Fix Static Analysis Error

diff --git a/Engine/Source/Editor/PhysicsAssetEditor/Private/PhysicsAssetDetailsCustomization.cpp b/Engine/Source/Editor/PhysicsAssetEditor/Private/PhysicsAssetDetailsCustomization.cpp
--- a/Engine/Source/Editor/PhysicsAssetEditor/Private/PhysicsAssetDetailsCustomization.cpp
+++ b/Engine/Source/Editor/PhysicsAssetEditor/Private/PhysicsAssetDetailsCustomization.cpp
@@ -218,7 +218,7 @@ TSharedRef< SWidget > FPhysicsAssetDetailsCustomization::FillPhysicalAnimationPr
 
 				Action.GetActionCheckState = FGetActionCheckState::CreateLambda([SharedData, ProfileName]()
 				{
-					return SharedData->PhysicsAsset->CurrentPhysicalAnimationProfileName == ProfileName ? ECheckBoxState::Checked : ECheckBoxState::Unchecked;
+					return (SharedData->PhysicsAsset->CurrentPhysicalAnimationProfileName == ProfileName) ? ECheckBoxState::Checked : ECheckBoxState::Unchecked;
 				});
 
 				TSharedRef<SWidget> PhysAnimProfileButton = SNew(STextBlock)
@@ -311,7 +311,7 @@ TSharedRef< SWidget > FPhysicsAssetDetailsCustomization::FillConstraintProfilesO
 
 				Action.GetActionCheckState = FGetActionCheckState::CreateLambda([SharedData, ProfileName]()
 				{
-					return SharedData->PhysicsAsset->CurrentConstraintProfileName == ProfileName ? ECheckBoxState::Checked : ECheckBoxState::Unchecked;
+					return (SharedData->PhysicsAsset->CurrentConstraintProfileName == ProfileName) ? ECheckBoxState::Checked : ECheckBoxState::Unchecked;
 				});
 
 				TSharedRef<SWidget> ConstraintProfileButton = SNew(STextBlock)
@@ -459,7 +459,7 @@ TSharedRef<SWidget> FPhysicsAssetDetailsCustomization::MakePhysicalAnimationProf
 						{
 							FName ProfileName = LocalPhysicsAssetEditorPtr.Pin()->GetSharedData()->PhysicsAsset->CurrentPhysicalAnimationProfileName;
 
-							return ProfileName == NAME_None ? FStyleColors::Foreground : FStyleColors::White;
+							return (ProfileName == NAME_None) ? FStyleColors::Foreground : FStyleColors::White;
 						})
 						.IsEnabled_Lambda([this]()
 						{
@@ -554,7 +554,7 @@ TSharedRef<SWidget> FPhysicsAssetDetailsCustomization::MakeConstraintProfilesWid
 						{
 							FName ProfileName = LocalPhysicsAssetEditorPtr.Pin()->GetSharedData()->PhysicsAsset->CurrentConstraintProfileName;
 
-							return ProfileName == NAME_None ? FStyleColors::Foreground : FStyleColors::White;
+							return (ProfileName == NAME_None) ? FStyleColors::Foreground : FStyleColors::White;
 						})
 						.Text_Lambda([LocalPhysicsAssetEditorPtr]()
 						{
