Physics Asset editor fixes
- Refresh physics asset editor hierachy when pasting bodies from clipboard
- Copy of only constraints now works
- Change wording of the copy a/ paste commands to reflect teh body constraint dulaity
- Fix duplicate menu entry for constraint context menu

diff --git a/Engine/Source/Editor/PhysicsAssetEditor/Private/PhysicsAssetEditor.cpp b/Engine/Source/Editor/PhysicsAssetEditor/Private/PhysicsAssetEditor.cpp
--- a/Engine/Source/Editor/PhysicsAssetEditor/Private/PhysicsAssetEditor.cpp
+++ b/Engine/Source/Editor/PhysicsAssetEditor/Private/PhysicsAssetEditor.cpp
@@ -1444,8 +1444,8 @@ void FPhysicsAssetEditor::BuildMenuWidgetConstraint(FMenuBuilder& InMenuBuilder)
 			FNewMenuDelegate::CreateStatic( &FLocal::FillAxesAndLimitsMenu ) );	
 		InMenuBuilder.AddSubMenu( LOCTEXT("ConvertMenu", "Convert"), LOCTEXT("ConvertMenu_ToolTip", "Convert constraint to various presets"),
 			FNewMenuDelegate::CreateStatic( &FLocal::FillConvertMenu ) );
-		InMenuBuilder.AddMenuEntry(Commands.CopyProperties);
-		InMenuBuilder.AddMenuEntry(Commands.PasteProperties);
+		InMenuBuilder.AddMenuEntry(Commands.CopyBodies);
+		InMenuBuilder.AddMenuEntry(Commands.PasteBodies);
 		InMenuBuilder.AddMenuEntry(Commands.CopyProperties);
 		InMenuBuilder.AddMenuEntry(Commands.PasteProperties);
 		InMenuBuilder.AddMenuEntry(Commands.DeleteConstraint);
@@ -2096,7 +2096,7 @@ bool FPhysicsAssetEditor::CanCopyBodies() const
 {
 	if (IsSelectedEditMode())
 	{
-		return (SharedData->SelectedBodies.Num() > 0);
+		return ((SharedData->SelectedBodies.Num() > 0) || (SharedData->SelectedConstraints.Num() > 0));
 	}
 	return false;
 }
diff --git a/Engine/Source/Editor/PhysicsAssetEditor/Private/PhysicsAssetEditorActions.cpp b/Engine/Source/Editor/PhysicsAssetEditor/Private/PhysicsAssetEditorActions.cpp
--- a/Engine/Source/Editor/PhysicsAssetEditor/Private/PhysicsAssetEditorActions.cpp
+++ b/Engine/Source/Editor/PhysicsAssetEditor/Private/PhysicsAssetEditorActions.cpp
@@ -10,8 +10,8 @@ void FPhysicsAssetEditorCommands::RegisterCommands()
 	UI_COMMAND(AddBodies, "Add/Replace Bodies", "Adds or replaces bodies for the selected bones using the current generation settings (see the Tools tab)", EUserInterfaceActionType::Button, FInputChord());
 
 	UI_COMMAND(ApplyPhysicalMaterial, "Physical Material", "Apply a Physical Material to All Bodies", EUserInterfaceActionType::Button, FInputChord());
-	UI_COMMAND(CopyBodies, "Copy Selected Bodies To Clipboard", "Copy selected bodies to clipboard", EUserInterfaceActionType::Button, FInputChord(EModifierKey::Control | EModifierKey::Shift, EKeys::C));
-	UI_COMMAND(PasteBodies, "Paste Bodies From Clipboard", "Paste selected bodies from clipboard", EUserInterfaceActionType::Button, FInputChord(EModifierKey::Control | EModifierKey::Shift, EKeys::V));
+	UI_COMMAND(CopyBodies, "Copy Selected Bodies/Constraints To Clipboard", "Copy selected bodies/constraints to clipboard", EUserInterfaceActionType::Button, FInputChord(EModifierKey::Control | EModifierKey::Shift, EKeys::C));
+	UI_COMMAND(PasteBodies, "Paste Bodies/Constraints From Clipboard", "Paste selected bodies/constraints from clipboard", EUserInterfaceActionType::Button, FInputChord(EModifierKey::Control | EModifierKey::Shift, EKeys::V));
 	UI_COMMAND(CopyProperties, "Copy Properties", "Copy Properties: Copy Properties Of Currently Selected Object To Next Selected Object", EUserInterfaceActionType::Button, FInputChord(EModifierKey::Control, EKeys::C));
 	UI_COMMAND(PasteProperties, "Paste Properties", "Paste Properties: Copy Properties Of Currently Selected Object To Next Selected Object", EUserInterfaceActionType::Button, FInputChord(EModifierKey::Control, EKeys::V));
 	UI_COMMAND(RepeatLastSimulation, "Simulate", "Previews Physics Simulation", EUserInterfaceActionType::RadioButton, FInputChord());
diff --git a/Engine/Source/Editor/PhysicsAssetEditor/Private/PhysicsAssetEditorSharedData.cpp b/Engine/Source/Editor/PhysicsAssetEditor/Private/PhysicsAssetEditorSharedData.cpp
--- a/Engine/Source/Editor/PhysicsAssetEditor/Private/PhysicsAssetEditorSharedData.cpp
+++ b/Engine/Source/Editor/PhysicsAssetEditor/Private/PhysicsAssetEditorSharedData.cpp
@@ -1617,6 +1617,7 @@ void FPhysicsAssetEditorSharedData::PasteBodiesAndConstraintsFromClipboard(int32
 			ClearSelectedBody();	//paste can change the primitives on our selected bodies. There's probably a way to properly update this, but for now just deselect
 			ClearSelectedConstraints();	//paste can change the primitives on our selected bodies. There's probably a way to properly update this, but for now just deselect
 			BroadcastPreviewChanged();
+			BroadcastHierarchyChanged();
 		}
 	}
 }
