Fix merge error that duplicated line
#rnx

diff --git a/Engine/Source/Editor/PhysicsAssetEditor/Private/PhysicsAssetEditorSharedData.cpp b/Engine/Source/Editor/PhysicsAssetEditor/Private/PhysicsAssetEditorSharedData.cpp
--- a/Engine/Source/Editor/PhysicsAssetEditor/Private/PhysicsAssetEditorSharedData.cpp
+++ b/Engine/Source/Editor/PhysicsAssetEditor/Private/PhysicsAssetEditorSharedData.cpp
@@ -2600,7 +2600,6 @@ void FPhysicsAssetEditorSharedData::EnableSimulation(bool bEnableSimulation)
 			EditorSkelComp->ResetAllBodiesSimulatePhysics();
 			EditorSkelComp->SetPhysicsBlendWeight(EditorOptions->PhysicsBlend);
 			PhysicalAnimationComponent->SetSkeletalMeshComponent(EditorSkelComp);
-		PhysicalAnimationComponent->SetSkeletalMeshComponent(EditorSkelComp);
 			// Make it start simulating
 			EditorSkelComp->WakeAllRigidBodies();
 		}
