[UE][FEATURE] Allow floor collisions to be disabled during simulation in PhAT
- Added a Floor Collision Enabled Checkbox to the 3-dots options menu attached to
the Simulate Tool bar button in PhAT.

diff --git a/Engine/Source/Editor/EditorStyle/Private/StarshipStyle.cpp b/Engine/Source/Editor/EditorStyle/Private/StarshipStyle.cpp
--- a/Engine/Source/Editor/EditorStyle/Private/StarshipStyle.cpp
+++ b/Engine/Source/Editor/EditorStyle/Private/StarshipStyle.cpp
@@ -7016,6 +7016,7 @@ void FStarshipEditorStyle::FStyle::SetupToolkitStyles()
 		Set( "PhysicsAssetEditor.EditingMode_Constraint.Small", new IMAGE_BRUSH( "/PhysicsAssetEditor/icon_PHatMode_Joint_40x", Icon20x20) );
 
 		Set( "PhysicsAssetEditor.SimulationNoGravity", new IMAGE_BRUSH( "PhysicsAssetEditor/icon_PhAT_PlaySimNoGravity_40x", Icon40x40 ) );
+		Set( "PhysicsAssetEditor.SimulationFloorCollision", new IMAGE_BRUSH("PhysicsAssetEditor/icon_PhAT_EnableCollision_40x", Icon40x40));
 		Set( "PhysicsAssetEditor.SelectedSimulation", new IMAGE_BRUSH( "PhysicsAssetEditor/icon_PhAT_PlaySimSelected_40x", Icon40x40 ) );
 		Set( "PhysicsAssetEditor.SimulationAll", new IMAGE_BRUSH_SVG("Starship/MainToolbar/simulate", Icon20x20));
 		Set( "PhysicsAssetEditor.Undo", new CORE_IMAGE_BRUSH_SVG("Starship/Common/Undo", Icon20x20));
diff --git a/Engine/Source/Editor/PhysicsAssetEditor/Private/PhysicsAssetEditor.cpp b/Engine/Source/Editor/PhysicsAssetEditor/Private/PhysicsAssetEditor.cpp
--- a/Engine/Source/Editor/PhysicsAssetEditor/Private/PhysicsAssetEditor.cpp
+++ b/Engine/Source/Editor/PhysicsAssetEditor/Private/PhysicsAssetEditor.cpp
@@ -573,6 +573,7 @@ void FPhysicsAssetEditor::ExtendToolbar()
 			MenuBuilder.BeginSection("SimulationOptions", LOCTEXT("SimulationOptionsHeader", "Simulation Options"));
 			{
 				MenuBuilder.AddMenuEntry(Commands.SimulationNoGravity);
+				MenuBuilder.AddMenuEntry(Commands.SimulationFloorCollision);
 			}
 			MenuBuilder.EndSection();
 
@@ -766,6 +767,12 @@ void FPhysicsAssetEditor::BindCommands()
 		FCanExecuteAction(),
 		FIsActionChecked::CreateSP(this, &FPhysicsAssetEditor::IsNoGravitySimulationEnabled));
 
+	ToolkitCommands->MapAction(
+		Commands.SimulationFloorCollision,
+		FExecuteAction::CreateSP(this, &FPhysicsAssetEditor::OnToggleSimulationFloorCollision),
+		FCanExecuteAction(),
+		FIsActionChecked::CreateSP(this, &FPhysicsAssetEditor::IsSimulationFloorCollisionEnabled));
+
 	ToolkitCommands->MapAction(
 		Commands.SelectedSimulation,
 		FExecuteAction::CreateSP(this, &FPhysicsAssetEditor::OnToggleSimulation, true),
@@ -2271,6 +2278,34 @@ bool FPhysicsAssetEditor::IsNoGravitySimulationEnabled() const
 	return SharedData->bNoGravitySimulation;
 }
 
+void FPhysicsAssetEditor::OnToggleSimulationFloorCollision()
+{
+	SharedData->bSimulationFloorCollisionEnabled = !SharedData->bSimulationFloorCollisionEnabled;
+
+	// Update collision for floor
+	if (PersonaToolkit)
+	{
+		TSharedRef<IPersonaPreviewScene> PersonaPreviewScene = PersonaToolkit->GetPreviewScene();
+
+		if (UStaticMeshComponent* FloorMeshComponent = const_cast<UStaticMeshComponent*>(PersonaPreviewScene->GetFloorMeshComponent()))
+		{
+			if (SharedData->bSimulationFloorCollisionEnabled)
+			{
+				FloorMeshComponent->SetCollisionEnabled(ECollisionEnabled::QueryAndPhysics);
+			}
+			else
+			{
+				FloorMeshComponent->SetCollisionEnabled(ECollisionEnabled::NoCollision);
+			}
+		}
+	}
+}
+
+bool FPhysicsAssetEditor::IsSimulationFloorCollisionEnabled() const
+{
+	return SharedData->bSimulationFloorCollisionEnabled;
+}
+
 bool FPhysicsAssetEditor::IsFullSimulation() const
 {
 	return !SelectedSimulation;
diff --git a/Engine/Source/Editor/PhysicsAssetEditor/Private/PhysicsAssetEditor.h b/Engine/Source/Editor/PhysicsAssetEditor/Private/PhysicsAssetEditor.h
--- a/Engine/Source/Editor/PhysicsAssetEditor/Private/PhysicsAssetEditor.h
+++ b/Engine/Source/Editor/PhysicsAssetEditor/Private/PhysicsAssetEditor.h
@@ -224,6 +224,8 @@ private:
 	void OnToggleSimulation(bool bInSelected);
 	void OnToggleSimulationNoGravity();
 	bool IsNoGravitySimulationEnabled() const;
+	void OnToggleSimulationFloorCollision();
+	bool IsSimulationFloorCollisionEnabled() const;
 	void SetupSelectedSimulation();
 	bool IsFullSimulation() const;
 	bool IsSelectedSimulation() const;
diff --git a/Engine/Source/Editor/PhysicsAssetEditor/Private/PhysicsAssetEditorActions.cpp b/Engine/Source/Editor/PhysicsAssetEditor/Private/PhysicsAssetEditorActions.cpp
--- a/Engine/Source/Editor/PhysicsAssetEditor/Private/PhysicsAssetEditorActions.cpp
+++ b/Engine/Source/Editor/PhysicsAssetEditor/Private/PhysicsAssetEditorActions.cpp
@@ -16,6 +16,7 @@ void FPhysicsAssetEditorCommands::RegisterCommands()
 	UI_COMMAND(PasteProperties, "Paste Properties", "Paste Properties: Copy Properties Of Currently Selected Object To Next Selected Object", EUserInterfaceActionType::Button, FInputChord(EModifierKey::Control, EKeys::V));
 	UI_COMMAND(RepeatLastSimulation, "Simulate", "Previews Physics Simulation", EUserInterfaceActionType::RadioButton, FInputChord());
 	UI_COMMAND(SimulationNoGravity, "No Gravity Simulation", "Run Physics Simulation without gravity. Use this to debug issues with your ragdoll. If the setup is correct, the asset should not move!", EUserInterfaceActionType::ToggleButton, FInputChord());
+	UI_COMMAND(SimulationFloorCollision, "Enable Floor Collisions", "Run Physics Simulation with collisions between physics bodies and the floor enabled.", EUserInterfaceActionType::ToggleButton, FInputChord());
 	UI_COMMAND(SelectedSimulation, "Simulate Selected", "Run Physics Simulation on selected objects. Use this to tune  specific parts of your ragdoll.", EUserInterfaceActionType::RadioButton, FInputChord(EModifierKey::Alt | EModifierKey::Shift, EKeys::I));
 	UI_COMMAND(SimulationAll, "Simulate", "Run Physics Simulation on all objects.", EUserInterfaceActionType::RadioButton, FInputChord(EModifierKey::Alt, EKeys::I));
 
diff --git a/Engine/Source/Editor/PhysicsAssetEditor/Private/PhysicsAssetEditorActions.h b/Engine/Source/Editor/PhysicsAssetEditor/Private/PhysicsAssetEditorActions.h
--- a/Engine/Source/Editor/PhysicsAssetEditor/Private/PhysicsAssetEditorActions.h
+++ b/Engine/Source/Editor/PhysicsAssetEditor/Private/PhysicsAssetEditorActions.h
@@ -31,6 +31,7 @@ public:
 	TSharedPtr<FUICommandInfo> PasteProperties;
 	TSharedPtr<FUICommandInfo> RepeatLastSimulation;
 	TSharedPtr<FUICommandInfo> SimulationNoGravity;
+	TSharedPtr<FUICommandInfo> SimulationFloorCollision;
 	TSharedPtr<FUICommandInfo> SelectedSimulation;
 	TSharedPtr<FUICommandInfo> SimulationAll;
 	TSharedPtr<FUICommandInfo> MeshRenderingMode_Solid;
diff --git a/Engine/Source/Editor/PhysicsAssetEditor/Private/PhysicsAssetEditorSharedData.cpp b/Engine/Source/Editor/PhysicsAssetEditor/Private/PhysicsAssetEditorSharedData.cpp
--- a/Engine/Source/Editor/PhysicsAssetEditor/Private/PhysicsAssetEditorSharedData.cpp
+++ b/Engine/Source/Editor/PhysicsAssetEditor/Private/PhysicsAssetEditorSharedData.cpp
@@ -74,6 +74,7 @@ FPhysicsAssetEditorSharedData::FPhysicsAssetEditorSharedData()
 
 	bRunningSimulation = false;
 	bNoGravitySimulation = false;
+	bSimulationFloorCollisionEnabled = true;
 
 	bManipulating = false;
 
diff --git a/Engine/Source/Editor/PhysicsAssetEditor/Private/PhysicsAssetEditorSharedData.h b/Engine/Source/Editor/PhysicsAssetEditor/Private/PhysicsAssetEditorSharedData.h
--- a/Engine/Source/Editor/PhysicsAssetEditor/Private/PhysicsAssetEditorSharedData.h
+++ b/Engine/Source/Editor/PhysicsAssetEditor/Private/PhysicsAssetEditorSharedData.h
@@ -349,7 +349,8 @@ public:
 	/** Misc toggles */
 	bool bRunningSimulation;
 	bool bNoGravitySimulation;
-	
+	bool bSimulationFloorCollisionEnabled;
+
 	/** Manipulation (rotate, translate, scale) */
 	bool bManipulating;
 
