[UE][FEATURE] Expand physics asset summary text in PhAT viewport
- Text now includes the number of potential colliding pairs (calculated as (bodies x (bodies -1)) - number of ignored collision pairs), thee number and type of each primative shape and the number of cross constraints.

diff --git a/Engine/Source/Editor/PhysicsAssetEditor/Private/PhysicsAssetEditorMode.cpp b/Engine/Source/Editor/PhysicsAssetEditor/Private/PhysicsAssetEditorMode.cpp
--- a/Engine/Source/Editor/PhysicsAssetEditor/Private/PhysicsAssetEditorMode.cpp
+++ b/Engine/Source/Editor/PhysicsAssetEditor/Private/PhysicsAssetEditorMode.cpp
@@ -8,6 +8,7 @@
 #include "ISkeletonEditorModule.h"
 #include "PhysicsAssetGraphSummoner.h"
 #include "PhysicsEngine/PhysicsAsset.h"
+#include "PhysicsEngine/PhysicsConstraintTemplate.h"
 #include "PhysicsAssetEditorActions.h"
 #include "SEditorViewportToolBarMenu.h"
 #include "PhysicsAssetEditorProfilesSummoner.h"
@@ -31,6 +32,130 @@ static const FName PhysicsAssetEditorProfilesName("PhysicsAssetProfilesView");
 static const FName PhysicsAssetEditorToolsName("PhysicsAssetTools");
 static const FName PhysicsAssetEditorAdvancedPreviewName(TEXT("AdvancedPreviewTab"));
 
+
+FText BuildPhysicsAssetShapeTypeCountText(const UPhysicsAsset* const PhysicsAsset)
+{
+	// Find the number of each type of shape in the asset and record them in a string.
+	if (PhysicsAsset)
+	{
+		// Find the number of each type of primitive in the physics asset.
+		TMap<EAggCollisionShape::Type, int32> ShapeCount;
+		ShapeCount.Add(EAggCollisionShape::Sphere, 0);
+		ShapeCount.Add(EAggCollisionShape::Box, 0);
+		ShapeCount.Add(EAggCollisionShape::Sphyl, 0);
+		ShapeCount.Add(EAggCollisionShape::Convex, 0);
+		ShapeCount.Add(EAggCollisionShape::TaperedCapsule, 0);
+
+		int32 TotalShapeCount = 0;
+
+		for (const TObjectPtr <USkeletalBodySetup>& SkeletalBodySetup : PhysicsAsset->SkeletalBodySetups)
+		{
+			if (SkeletalBodySetup)
+			{
+				for (auto& Element : ShapeCount)
+				{
+					const int32 Count = SkeletalBodySetup->AggGeom.GetElementCount(Element.Key);
+					Element.Value += Count;
+					TotalShapeCount += Count;
+				}
+			}
+		}
+
+		// Create a text element for each value that will be included in the output.
+		FFormatOrderedArguments FormatArgs;
+
+		auto AddToFormatArgs = [&ShapeCount, &FormatArgs](const EAggCollisionShape::Type InShapeType, const FTextFormat& InTextFormat)
+		{
+			if (const int32 CurrentShapeTypeCount = ShapeCount[InShapeType])
+			{
+				FormatArgs.Add(FText::Format(InTextFormat, CurrentShapeTypeCount));
+			}
+		};
+
+		FormatArgs.Add(FText::Format(LOCTEXT("PrimitiveShapeCount", "{0} {0}|plural(one=Primitive,other=Primitives)"), TotalShapeCount));
+
+		AddToFormatArgs(EAggCollisionShape::Sphere, LOCTEXT("PrimitiveSphereCount", "{0} {0}|plural(one=Sphere,other=Spheres)"));
+		AddToFormatArgs(EAggCollisionShape::Box, LOCTEXT("PrimitiveBoxCount", "{0} {0}|plural(one=Box,other=Boxes)"));
+		AddToFormatArgs(EAggCollisionShape::Sphyl, LOCTEXT("PrimitiveSphylCount", "{0} {0}|plural(one=Capsule,other=Capsules)"));
+		AddToFormatArgs(EAggCollisionShape::Convex, LOCTEXT("PrimitiveConvexCount", "{0} {0}|plural(one=Convex Element,other=Convex Elements)"));
+		AddToFormatArgs(EAggCollisionShape::TaperedCapsule, LOCTEXT("PrimitiveTaperedCapsuleCount", "{0} {0}|plural(one=Tapered Capsule,other=Tapered Capsules)"));
+
+		// Build format string for the number of accumulated elements.
+		FString FormatStr = "{0}";
+
+		if (FormatArgs.Num() > 1)
+		{
+			FormatStr += ": (";
+
+			for (uint32 Index = 1, Max = FormatArgs.Num(); Index < Max; ++Index)
+			{
+				FormatStr += "{" + FString::FromInt(Index) + "}, ";
+			}
+
+			FormatStr.RemoveFromEnd(" ");
+			FormatStr.RemoveFromEnd(",");
+			FormatStr += ")";
+		}
+
+		return FText::Format(FText::FromString(FormatStr), FormatArgs);
+	}
+
+	return FText();
+}
+
+FText BuildCrossConstraintCountText(const UPhysicsAsset* const PhysicsAsset)
+{
+	// Find the number of cross constraints in the asset and record them in a string.
+	if (PhysicsAsset)
+	{
+		uint32 CrossConstraintCount = 0;
+
+		const USkeletalMesh* const SkeletalMesh = PhysicsAsset->GetPreviewMesh();
+		if (SkeletalMesh)
+		{
+			const FReferenceSkeleton& RefSkeleton = SkeletalMesh->GetRefSkeleton();
+
+			for (TObjectPtr<UPhysicsConstraintTemplate> CurrentConstraintSetup : PhysicsAsset->ConstraintSetup)
+			{
+				if (CurrentConstraintSetup)
+				{
+					const int32 ChildBoneIndex = RefSkeleton.FindBoneIndex(CurrentConstraintSetup->DefaultInstance.GetChildBoneName());
+					const int32 ParentBoneIndex = RefSkeleton.FindBoneIndex(CurrentConstraintSetup->DefaultInstance.GetParentBoneName());
+
+					if (ParentBoneIndex != RefSkeleton.GetParentIndex(ChildBoneIndex))
+					{
+						++CrossConstraintCount;
+					}
+				}
+			}
+		}
+
+		if (CrossConstraintCount > 0)
+		{
+			return FText::Format(LOCTEXT("CrossConstraintCount", "({0} {0}|plural(one=Cross Constraint,other=Cross Constraints))"), CrossConstraintCount);
+		}
+	}
+
+	return FText();
+}
+
+uint32 CalculateTotalNumberOfCollisionIteractions(const UPhysicsAsset* const PhysicsAsset)
+{
+	uint32 Result = 0;
+
+	// Find total number of collision iteractions.
+	if (PhysicsAsset)
+	{
+		const uint32 BodyCount = PhysicsAsset->SkeletalBodySetups.Num();
+		const uint32 PotentialCollisionCount = BodyCount * (BodyCount - 1);
+		const uint32 IgnoredCollisionCount = PhysicsAsset->CollisionDisableTable.Num();
+		Result = PotentialCollisionCount - IgnoredCollisionCount;
+	}
+
+	return Result;
+}
+
+
 FPhysicsAssetEditorMode::FPhysicsAssetEditorMode(TSharedRef<FWorkflowCentricApplication> InHostingApp, TSharedRef<ISkeletonTree> InSkeletonTree, TSharedRef<IPersonaPreviewScene> InPreviewScene)
 	: FApplicationMode(PhysicsAssetEditorModes::PhysicsAssetEditorMode)
 {
@@ -188,13 +313,16 @@ FPhysicsAssetEditorMode::FPhysicsAssetEditorMode(TSharedRef<FWorkflowCentricAppl
 		{
 			TSharedPtr<FPhysicsAssetEditorSharedData> SharedData = PhysicsAssetEditorPtr.Pin()->GetSharedData();
 
-			// Write body/constraint count at top.
+			// Write physics asset summary at the top of the view port.
 			return FText::Format(
-				NSLOCTEXT("UnrealEd", "BodiesConstraints_F", "{0} Bodies, {1} Considered for bounds ({2}%), {3} Constraints"),
+				NSLOCTEXT("UnrealEd", "BodiesConstraints_F", "{0} Bodies ({1} Considered for bounds, {2}%)\n{3}\n{4} Constraints{5}\n{6} Collision Interactions"),
 				FText::AsNumber(SharedData->PhysicsAsset->SkeletalBodySetups.Num()),
 				FText::AsNumber(SharedData->PhysicsAsset->BoundsBodies.Num()),
 				FText::AsNumber(static_cast<float>(SharedData->PhysicsAsset->BoundsBodies.Num()) / static_cast<float>(SharedData->PhysicsAsset->SkeletalBodySetups.Num()) * 100.0f),
-				FText::AsNumber(SharedData->PhysicsAsset->ConstraintSetup.Num()));
+				BuildPhysicsAssetShapeTypeCountText(SharedData->PhysicsAsset),
+				FText::AsNumber(SharedData->PhysicsAsset->ConstraintSetup.Num()),
+				BuildCrossConstraintCountText(SharedData->PhysicsAsset),
+				FText::AsNumber(CalculateTotalNumberOfCollisionIteractions(SharedData->PhysicsAsset)));
 		}
 
 		return FText();
