PhysicsAssetDetailsCustomization: Fix duplicate localization keys

diff --git a/Engine/Source/Editor/PhysicsAssetEditor/Private/PhysicsAssetDetailsCustomization.cpp b/Engine/Source/Editor/PhysicsAssetEditor/Private/PhysicsAssetDetailsCustomization.cpp
--- a/Engine/Source/Editor/PhysicsAssetEditor/Private/PhysicsAssetDetailsCustomization.cpp
+++ b/Engine/Source/Editor/PhysicsAssetEditor/Private/PhysicsAssetDetailsCustomization.cpp
@@ -156,17 +156,17 @@ TSharedRef< SWidget > FPhysicsAssetDetailsCustomization::FillPhysicalAnimationPr
 
 	if(SharedData->PhysicsAsset)
 	{
-		MenuBuilder.BeginSection("NewProfile", LOCTEXT("PhysicsAssetEditor_NewPhysicalAnimationMenu", "New"));
+		MenuBuilder.BeginSection("NewPhysicalAnimationProfile", LOCTEXT("PhysicsAssetEditor_NewPhysicalAnimationMenu", "New"));
 		{
 			MenuBuilder.AddMenuEntry(Commands.NewPhysicalAnimationProfile);
 		}
 		MenuBuilder.EndSection();
 
-		MenuBuilder.BeginSection("CurrentProfile", LOCTEXT("PhysicsAssetEditor_CurrentPhysicalAnimationMenu", "Current Profile"));
+		MenuBuilder.BeginSection("CurrentPhysicalAnimationProfile", LOCTEXT("PhysicsAssetEditor_CurrentPhysicalAnimationMenu", "Current Profile"));
 		{
 			MenuBuilder.AddMenuEntry(
-				LOCTEXT("PhysicsAssetEditor_RenameMenu", "Rename"),
-				LOCTEXT("PhysicsAssetEditor_RenameMenuTooltip", "Rename the Current Physical Animation Profile"),
+				LOCTEXT("PhysicsAssetEditor_RenamePhysicalAnimationMenu", "Rename"),
+				LOCTEXT("PhysicsAssetEditor_RenamePhysicalAnimationTooltip", "Rename the Current Physical Animation Profile"),
 				FSlateIcon(),
 				FUIAction(
 					FExecuteAction::CreateLambda([this]
@@ -249,17 +249,17 @@ TSharedRef< SWidget > FPhysicsAssetDetailsCustomization::FillConstraintProfilesO
 
 	if(SharedData->PhysicsAsset)
 	{
-		MenuBuilder.BeginSection("NewProfile", LOCTEXT("PhysicsAssetEditor_NewConstraintProfileMenu", "New"));
+		MenuBuilder.BeginSection("NewConstraintProfile", LOCTEXT("PhysicsAssetEditor_NewConstraintProfileMenu", "New"));
 		{
 			MenuBuilder.AddMenuEntry(Commands.NewConstraintProfile);
 		}
 		MenuBuilder.EndSection();
 
-		MenuBuilder.BeginSection("CurrentProfile", LOCTEXT("PhysicsAssetEditor_CurrentProfileMenu", "Current Profile"));
+		MenuBuilder.BeginSection("CurrentConstraintProfile", LOCTEXT("PhysicsAssetEditor_CurrentConstraintProfileMenu", "Current Profile"));
 		{
 			MenuBuilder.AddMenuEntry(
-				LOCTEXT("PhysicsAssetEditor_RenameMenu", "Rename"),
-				LOCTEXT("PhysicsAssetEditor_RenameMenuTooltip", "Rename the Current Constraint Profile"),
+				LOCTEXT("PhysicsAssetEditor_RenameConstraintMenu", "Rename"),
+				LOCTEXT("PhysicsAssetEditor_RenameConstraintTooltip", "Rename the Current Constraint Profile"),
 				FSlateIcon(),
 				FUIAction(
 					FExecuteAction::CreateLambda([this]
@@ -429,7 +429,7 @@ TSharedRef<SWidget> FPhysicsAssetDetailsCustomization::MakePhysicalAnimationProf
 		.Padding(0.0f, 0.0f, 2.0f, 3.0f)
 		[
 			SNew(STextBlock)
-			.Text(LOCTEXT("CurrentProfile", "Current Profile"))
+			.Text(LOCTEXT("CurrentPhysicalAnimationProfile", "Current Profile"))
 		]
 
 		+SHorizontalBox::Slot()
@@ -471,7 +471,7 @@ TSharedRef<SWidget> FPhysicsAssetDetailsCustomization::MakePhysicalAnimationProf
 							if(LocalPhysicsAssetEditorPtr.Pin()->GetSharedData()->PhysicsAsset->CurrentPhysicalAnimationProfileName != ProfileAsName &&
 								LocalPhysicsAssetEditorPtr.Pin()->GetSharedData()->PhysicsAsset->GetPhysicalAnimationProfileNames().Contains(ProfileAsName))
 							{
-								PhysicalAnimationProfileNameTextBox->SetError(FText::Format(LOCTEXT("ProfileExists", "Profile '{0}' already exists"), InText));
+								PhysicalAnimationProfileNameTextBox->SetError(FText::Format(LOCTEXT("PhysicalAnimationProfileExists", "Profile '{0}' already exists"), InText));
 							}
 							else
 							{
@@ -528,7 +528,7 @@ TSharedRef<SWidget> FPhysicsAssetDetailsCustomization::MakeConstraintProfilesWid
 		.Padding(0.0f, 0.0f, 2.0f, 3.0f)
 		[
 			SNew(STextBlock)
-			.Text(LOCTEXT("CurrentProfile", "Current Profile"))
+			.Text(LOCTEXT("CurrentConstraintProfile", "Current Profile"))
 		]
 
 		+SHorizontalBox::Slot()
@@ -570,7 +570,7 @@ TSharedRef<SWidget> FPhysicsAssetDetailsCustomization::MakeConstraintProfilesWid
 							if(LocalPhysicsAssetEditorPtr.Pin()->GetSharedData()->PhysicsAsset->CurrentConstraintProfileName != ProfileAsName &&
 								LocalPhysicsAssetEditorPtr.Pin()->GetSharedData()->PhysicsAsset->GetConstraintProfileNames().Contains(ProfileAsName))
 							{
-								ConstraintProfileNameTextBox->SetError(FText::Format(LOCTEXT("ProfileExists", "Profile '{0}' already exists"), InText));
+								ConstraintProfileNameTextBox->SetError(FText::Format(LOCTEXT("ConstraintProfileExists", "Profile '{0}' already exists"), InText));
 							}
 							else
 							{
