[UE][FIX] for collision interation count in PhAT asset summary text

diff --git a/Engine/Source/Editor/PhysicsAssetEditor/Private/PhysicsAssetEditorMode.cpp b/Engine/Source/Editor/PhysicsAssetEditor/Private/PhysicsAssetEditorMode.cpp
--- a/Engine/Source/Editor/PhysicsAssetEditor/Private/PhysicsAssetEditorMode.cpp
+++ b/Engine/Source/Editor/PhysicsAssetEditor/Private/PhysicsAssetEditorMode.cpp
@@ -147,7 +147,7 @@ uint32 CalculateTotalNumberOfCollisionIteractions(const UPhysicsAsset* const Phy
 	if (PhysicsAsset)
 	{
 		const uint32 BodyCount = PhysicsAsset->SkeletalBodySetups.Num();
-		const uint32 PotentialCollisionCount = BodyCount * (BodyCount - 1);
+		const uint32 PotentialCollisionCount = (BodyCount * (BodyCount - 1)) / 2;
 		const uint32 IgnoredCollisionCount = PhysicsAsset->CollisionDisableTable.Num();
 		Result = PotentialCollisionCount - IgnoredCollisionCount;
 	}
