[UE][FEATURE] PhAT - Move Simulation Floor Collisions option to Editor Options so that it persists across editor instances.

diff --git a/Engine/Source/Editor/PhysicsAssetEditor/Private/PhysicsAssetEditor.cpp b/Engine/Source/Editor/PhysicsAssetEditor/Private/PhysicsAssetEditor.cpp
--- a/Engine/Source/Editor/PhysicsAssetEditor/Private/PhysicsAssetEditor.cpp
+++ b/Engine/Source/Editor/PhysicsAssetEditor/Private/PhysicsAssetEditor.cpp
@@ -2280,22 +2280,25 @@ bool FPhysicsAssetEditor::IsNoGravitySimulationEnabled() const
 
 void FPhysicsAssetEditor::OnToggleSimulationFloorCollision()
 {
-	SharedData->bSimulationFloorCollisionEnabled = !SharedData->bSimulationFloorCollisionEnabled;
-
-	// Update collision for floor
-	if (PersonaToolkit)
+	if (SharedData && SharedData->EditorOptions)
 	{
-		TSharedRef<IPersonaPreviewScene> PersonaPreviewScene = PersonaToolkit->GetPreviewScene();
+		SharedData->EditorOptions->bSimulationFloorCollisionEnabled = !SharedData->EditorOptions->bSimulationFloorCollisionEnabled;
 
-		if (UStaticMeshComponent* FloorMeshComponent = const_cast<UStaticMeshComponent*>(PersonaPreviewScene->GetFloorMeshComponent()))
+		// Update collision for floor
+		if (PersonaToolkit)
 		{
-			if (SharedData->bSimulationFloorCollisionEnabled)
-			{
-				FloorMeshComponent->SetCollisionEnabled(ECollisionEnabled::QueryAndPhysics);
-			}
-			else
+			TSharedRef<IPersonaPreviewScene> PersonaPreviewScene = PersonaToolkit->GetPreviewScene();
+
+			if (UStaticMeshComponent* FloorMeshComponent = const_cast<UStaticMeshComponent*>(PersonaPreviewScene->GetFloorMeshComponent()))
 			{
-				FloorMeshComponent->SetCollisionEnabled(ECollisionEnabled::NoCollision);
+				if (SharedData->EditorOptions->bSimulationFloorCollisionEnabled)
+				{
+					FloorMeshComponent->SetCollisionEnabled(ECollisionEnabled::QueryAndPhysics);
+				}
+				else
+				{
+					FloorMeshComponent->SetCollisionEnabled(ECollisionEnabled::NoCollision);
+				}
 			}
 		}
 	}
@@ -2303,7 +2306,7 @@ void FPhysicsAssetEditor::OnToggleSimulationFloorCollision()
 
 bool FPhysicsAssetEditor::IsSimulationFloorCollisionEnabled() const
 {
-	return SharedData->bSimulationFloorCollisionEnabled;
+	return SharedData && SharedData->EditorOptions && SharedData->EditorOptions->bSimulationFloorCollisionEnabled;
 }
 
 bool FPhysicsAssetEditor::IsFullSimulation() const
diff --git a/Engine/Source/Editor/PhysicsAssetEditor/Private/PhysicsAssetEditorSharedData.cpp b/Engine/Source/Editor/PhysicsAssetEditor/Private/PhysicsAssetEditorSharedData.cpp
--- a/Engine/Source/Editor/PhysicsAssetEditor/Private/PhysicsAssetEditorSharedData.cpp
+++ b/Engine/Source/Editor/PhysicsAssetEditor/Private/PhysicsAssetEditorSharedData.cpp
@@ -74,7 +74,6 @@ FPhysicsAssetEditorSharedData::FPhysicsAssetEditorSharedData()
 
 	bRunningSimulation = false;
 	bNoGravitySimulation = false;
-	bSimulationFloorCollisionEnabled = true;
 
 	bManipulating = false;
 
diff --git a/Engine/Source/Editor/PhysicsAssetEditor/Private/PhysicsAssetEditorSharedData.h b/Engine/Source/Editor/PhysicsAssetEditor/Private/PhysicsAssetEditorSharedData.h
--- a/Engine/Source/Editor/PhysicsAssetEditor/Private/PhysicsAssetEditorSharedData.h
+++ b/Engine/Source/Editor/PhysicsAssetEditor/Private/PhysicsAssetEditorSharedData.h
@@ -349,7 +349,6 @@ public:
 	/** Misc toggles */
 	bool bRunningSimulation;
 	bool bNoGravitySimulation;
-	bool bSimulationFloorCollisionEnabled;
 
 	/** Manipulation (rotate, translate, scale) */
 	bool bManipulating;
diff --git a/Engine/Source/Editor/UnrealEd/Classes/Preferences/PhysicsAssetEditorOptions.h b/Engine/Source/Editor/UnrealEd/Classes/Preferences/PhysicsAssetEditorOptions.h
--- a/Engine/Source/Editor/UnrealEd/Classes/Preferences/PhysicsAssetEditorOptions.h
+++ b/Engine/Source/Editor/UnrealEd/Classes/Preferences/PhysicsAssetEditorOptions.h
@@ -111,6 +111,10 @@ class UNREALED_API UPhysicsAssetEditorOptions : public UObject
 	UPROPERTY(config)
 	uint32 bRenderOnlySelectedConstraints:1;
 
+	/* Toggle collisions with floor in the simulation */
+	UPROPERTY(config)
+	uint32 bSimulationFloorCollisionEnabled:1;
+
 	/** Controls how large constraints are drawn in Physics Asset Editor */
 	UPROPERTY(config)
 	float ConstraintDrawSize;
diff --git a/Engine/Source/Editor/UnrealEd/Private/PreferenceStubs.cpp b/Engine/Source/Editor/UnrealEd/Private/PreferenceStubs.cpp
--- a/Engine/Source/Editor/UnrealEd/Private/PreferenceStubs.cpp
+++ b/Engine/Source/Editor/UnrealEd/Private/PreferenceStubs.cpp
@@ -40,6 +40,7 @@ UPhysicsAssetEditorOptions::UPhysicsAssetEditorOptions(const FObjectInitializer&
 	InterpolationSpeed = 50.f;
 
 	bShowConstraintsAsPoints = false;
+	bSimulationFloorCollisionEnabled = true;
 	ConstraintDrawSize = 1.0f;
 
 	// view options
