Phat : fix invalid reported number of bodies and ocnstraints when pasting from nothing, also fixed typo

diff --git a/Engine/Source/Editor/PhysicsAssetEditor/Private/PhysicsAssetEditor.cpp b/Engine/Source/Editor/PhysicsAssetEditor/Private/PhysicsAssetEditor.cpp
--- a/Engine/Source/Editor/PhysicsAssetEditor/Private/PhysicsAssetEditor.cpp
+++ b/Engine/Source/Editor/PhysicsAssetEditor/Private/PhysicsAssetEditor.cpp
@@ -2114,7 +2114,7 @@ void FPhysicsAssetEditor::OnCopyBodies()
 	int32 NumCopiedConstraints;
 	SharedData->CopySelectedBodiesAndConstraintsToClipboard(NumCopiedBodies, NumCopiedConstraints);
 
-	const FText MessageFormat = LOCTEXT("CopiedBodiesAndConstraintsToClipboard", "{0} bodies ans {1} constraints copied to clipboard");
+	const FText MessageFormat = LOCTEXT("CopiedBodiesAndConstraintsToClipboard", "{0} bodies and {1} constraints copied to clipboard");
 	ShowNotificationMessage(FText::Format(MessageFormat, NumCopiedBodies, NumCopiedConstraints), SNotificationItem::CS_Success);
 }
 
@@ -2139,7 +2139,7 @@ void FPhysicsAssetEditor::OnPasteBodies()
 	int32 NumPastedConstraints;
 	SharedData->PasteBodiesAndConstraintsFromClipboard(NumPastedBodies, NumPastedConstraints);
 
-	const FText MessageFormat = LOCTEXT("PastedBodiesAndConstraintsToClipboard", "{0} bodies ans {1} constraints pasted from clipboard");
+	const FText MessageFormat = LOCTEXT("PastedBodiesAndConstraintsToClipboard", "{0} bodies and {1} constraints pasted from clipboard");
 	ShowNotificationMessage(FText::Format(MessageFormat, NumPastedBodies, NumPastedConstraints), SNotificationItem::CS_Success);
 }
 
diff --git a/Engine/Source/Editor/PhysicsAssetEditor/Private/PhysicsAssetEditorSharedData.cpp b/Engine/Source/Editor/PhysicsAssetEditor/Private/PhysicsAssetEditorSharedData.cpp
--- a/Engine/Source/Editor/PhysicsAssetEditor/Private/PhysicsAssetEditorSharedData.cpp
+++ b/Engine/Source/Editor/PhysicsAssetEditor/Private/PhysicsAssetEditorSharedData.cpp
@@ -1469,6 +1469,8 @@ void FPhysicsAssetEditorSharedData::AutoNamePrimitive(int32 BodyIndex, EAggColli
 
 void FPhysicsAssetEditorSharedData::CopySelectedBodiesAndConstraintsToClipboard(int32& OutNumCopiedBodies, int32& OutNumCopiedConstraints)
 {
+	OutNumCopiedBodies = 0;
+	OutNumCopiedConstraints = 0;
 	if (PhysicsAsset)
 	{
 		// Clear the mark state for saving.
@@ -1563,6 +1565,8 @@ public:
 
 void FPhysicsAssetEditorSharedData::PasteBodiesAndConstraintsFromClipboard(int32& OutNumPastedBodies, int32& OutNumPastedConstraints)
 {
+	OutNumPastedBodies = 0;
+	OutNumPastedConstraints = 0;
 	if (PhysicsAsset)
 	{
 		FString TextToImport;
