Shape selection in PHAT - selecting shapes in viewport selects them also in skeleton tree.
This appears to have already been implemented, but had 2 bugs.

PhysicsAssetEditor::FShapeData needs all data to represent a unique shape, but did not include primitive type. Primitive index is the index into the list of primitives of a type, not into the full list of primitives, so primitive type needed to be added to this struct.

FPhysicsAssetEditor::HandleViewportSelectionChanged was looking for shapes within the list of selected constraints rather than the list of selected bodies. This change adds code to check the body selection for shapes _as well as_ the constraint selection. I left the constraint selection bit in there in case we want selecting a constraint to select the shapes it's attached to.

diff --git a/Engine/Source/Editor/PhysicsAssetEditor/Private/PhysicsAssetEditor.cpp b/Engine/Source/Editor/PhysicsAssetEditor/Private/PhysicsAssetEditor.cpp
--- a/Engine/Source/Editor/PhysicsAssetEditor/Private/PhysicsAssetEditor.cpp
+++ b/Engine/Source/Editor/PhysicsAssetEditor/Private/PhysicsAssetEditor.cpp
@@ -104,21 +104,27 @@ namespace PhysicsAssetEditor
 	{
 		int32 Index;
 		int32 PrimitiveIndex;
+		EAggCollisionShape::Type PrimitiveType;
 
-		FShapeData(int32 Index, int32 PrimitiveIndex)
+		FShapeData(int32 Index, int32 PrimitiveIndex, EAggCollisionShape::Type PrimitiveType)
 			: Index(Index)
 			, PrimitiveIndex(PrimitiveIndex)
+			, PrimitiveType(PrimitiveType)
 		{
 		}
 
 		bool operator==(const FShapeData& rhs) const
 		{
-			return Index == rhs.Index && PrimitiveIndex == rhs.PrimitiveIndex;
+			return Index == rhs.Index && PrimitiveIndex == rhs.PrimitiveIndex && PrimitiveType == rhs.PrimitiveType;
 		}
 
 		friend uint32 GetTypeHash(const FShapeData& ShapeData)
 		{
-			return HashCombine(::GetTypeHash(ShapeData.Index), ::GetTypeHash(ShapeData.PrimitiveIndex));
+			return HashCombine(
+				HashCombine(
+					::GetTypeHash(ShapeData.Index),
+					::GetTypeHash(ShapeData.PrimitiveIndex)),
+				::GetTypeHash(ShapeData.PrimitiveType));
 		}
 	};
 }
@@ -294,9 +300,15 @@ void FPhysicsAssetEditor::HandleViewportSelectionChanged(const TArray<FPhysicsAs
 			{ 
 				return SharedData->PhysicsAsset->ConstraintSetup[InItem.Index];
 			});
+			// NOTE: Does selecting a constraint also select a shape?
+			// If not then this is not needed.
 			Algo::Transform(InSelectedConstraints, Shapes, [this](const FPhysicsAssetEditorSharedData::FSelection& InItem)
 			{
-				return PhysicsAssetEditor::FShapeData(InItem.Index, InItem.PrimitiveIndex);
+				return PhysicsAssetEditor::FShapeData(InItem.Index, InItem.PrimitiveIndex, InItem.PrimitiveType);
+			});
+			Algo::Transform(InSelectedBodies, Shapes, [this](const FPhysicsAssetEditorSharedData::FSelection& InItem)
+			{
+				return PhysicsAssetEditor::FShapeData(InItem.Index, InItem.PrimitiveIndex, InItem.PrimitiveType);
 			});
 
 			if (PhysAssetProperties.IsValid())
@@ -320,7 +332,7 @@ void FPhysicsAssetEditor::HandleViewportSelectionChanged(const TArray<FPhysicsAs
 					else if (InItem->IsOfType<FSkeletonTreePhysicsShapeItem>())
 					{
 						TSharedRef<FSkeletonTreePhysicsShapeItem> ShapeItem = StaticCastSharedRef<FSkeletonTreePhysicsShapeItem>(InItem);
-						PhysicsAssetEditor::FShapeData ShapeData(ShapeItem->GetBodySetupIndex(), ShapeItem->GetShapeIndex());
+						PhysicsAssetEditor::FShapeData ShapeData(ShapeItem->GetBodySetupIndex(), ShapeItem->GetShapeIndex(), ShapeItem->GetShapeType());
 						if (Shapes.Contains(ShapeData))
 						{
 							bInOutExpand = true;
