Add PHAT controls for setting and selecting per-shape collisionenabled settings. Specifically:
- Added menu options for setting shapes to QueryAndProbe and ProbeOnly
- Added PHAT actions for selecting shapes in a body based on collision type (useful for complex physics setups)

This change depends on 20366101, which adds the QueryAndProbe and ProbeOnly CollisionEnabled settings going in first.

diff --git a/Engine/Source/Editor/PhysicsAssetEditor/Private/PhysicsAssetDetailsCustomization.cpp b/Engine/Source/Editor/PhysicsAssetEditor/Private/PhysicsAssetDetailsCustomization.cpp
--- a/Engine/Source/Editor/PhysicsAssetEditor/Private/PhysicsAssetDetailsCustomization.cpp
+++ b/Engine/Source/Editor/PhysicsAssetEditor/Private/PhysicsAssetDetailsCustomization.cpp
@@ -795,7 +795,7 @@ void FPhysicsAssetDetailsCustomization::SelectAllBodiesInCurrentPhysicalAnimatio
 		}
 	}
 	SharedData->ClearSelectedBody();	//clear selection
-	SharedData->SetSelectedBodiesAnyPrim(NewBodiesSelection, true);
+	SharedData->SetSelectedBodiesAnyPrimitive(NewBodiesSelection, true);
 
 	return;
 }
diff --git a/Engine/Source/Editor/PhysicsAssetEditor/Private/PhysicsAssetEditor.cpp b/Engine/Source/Editor/PhysicsAssetEditor/Private/PhysicsAssetEditor.cpp
--- a/Engine/Source/Editor/PhysicsAssetEditor/Private/PhysicsAssetEditor.cpp
+++ b/Engine/Source/Editor/PhysicsAssetEditor/Private/PhysicsAssetEditor.cpp
@@ -504,7 +504,7 @@ void FPhysicsAssetEditor::OnFinishedChangingProperties(const FPropertyChangedEve
 		{ // bulk update
 			FScopedBulkSelection BulkSelection(SharedData);
 			SharedData->ClearSelectedBody();
-			SharedData->SetSelectedBodiesAnyPrim(SelectedBodyIndices, true);
+			SharedData->SetSelectedBodiesAnyPrimitive(SelectedBodyIndices, true);
 		}
 	}
 
@@ -692,13 +692,7 @@ void FPhysicsAssetEditor::ExtendMenu()
 		{
 			const FPhysicsAssetEditorCommands& Commands = FPhysicsAssetEditorCommands::Get();
 			MenuBarBuilder.BeginSection("Selection", LOCTEXT("PhatEditSelection", "Selection"));
-			MenuBarBuilder.AddMenuEntry(Commands.SelectAllBodies);
-			MenuBarBuilder.AddMenuEntry(Commands.SelectSimulatedBodies);
-			MenuBarBuilder.AddMenuEntry(Commands.SelectKinematicBodies);
-			MenuBarBuilder.AddMenuEntry(Commands.SelectAllConstraints);
-			MenuBarBuilder.AddMenuEntry(Commands.ToggleSelectionType);
-			MenuBarBuilder.AddMenuEntry(Commands.ToggleSelectionTypeWithUserConstraints);
-			MenuBarBuilder.AddMenuEntry(Commands.ToggleShowSelected);
+
 			MenuBarBuilder.AddMenuEntry(Commands.ShowSelected);
 			MenuBarBuilder.AddMenuEntry(Commands.HideSelected);
 			MenuBarBuilder.AddMenuEntry(Commands.ToggleShowOnlySelected);
@@ -707,6 +701,27 @@ void FPhysicsAssetEditor::ExtendMenu()
 			MenuBarBuilder.AddMenuEntry(Commands.ShowAll);
 			MenuBarBuilder.AddMenuEntry(Commands.HideAll);
 			MenuBarBuilder.AddMenuEntry(Commands.DeselectAll);
+			MenuBarBuilder.AddMenuEntry(Commands.ToggleShowSelected);
+
+			// Bodies
+			MenuBarBuilder.BeginSection("Bodies & Constraints", LOCTEXT("PhatEditSelectionBodies", "Bodies & Constraints"));
+				MenuBarBuilder.AddMenuEntry(Commands.SelectAllBodies);
+				MenuBarBuilder.AddMenuEntry(Commands.SelectSimulatedBodies);
+				MenuBarBuilder.AddMenuEntry(Commands.SelectKinematicBodies);
+				MenuBarBuilder.AddMenuEntry(Commands.SelectAllConstraints);
+				MenuBarBuilder.AddMenuEntry(Commands.ToggleSelectionType);
+				MenuBarBuilder.AddMenuEntry(Commands.ToggleSelectionTypeWithUserConstraints);
+			MenuBarBuilder.EndSection();
+
+			// Shapes
+			MenuBarBuilder.BeginSection("Shapes", LOCTEXT("PhatEditSelectionShapes", "Shapes"));
+				MenuBarBuilder.AddMenuEntry(Commands.SelectShapesQueryOnly);
+				MenuBarBuilder.AddMenuEntry(Commands.SelectShapesQueryAndPhysics);
+				MenuBarBuilder.AddMenuEntry(Commands.SelectShapesPhysicsOnly);
+				MenuBarBuilder.AddMenuEntry(Commands.SelectShapesQueryAndProbe);
+				MenuBarBuilder.AddMenuEntry(Commands.SelectShapesProbeOnly);
+			MenuBarBuilder.EndSection();
+
 			MenuBarBuilder.EndSection();
 		}
 	};
@@ -815,6 +830,12 @@ void FPhysicsAssetEditor::BindCommands()
 		FCanExecuteAction::CreateSP(this, &FPhysicsAssetEditor::CanSetPrimitiveCollision, ECollisionEnabled::QueryAndPhysics),
 		FIsActionChecked::CreateSP(this, &FPhysicsAssetEditor::IsPrimitiveCollisionChecked, ECollisionEnabled::QueryAndPhysics));
 
+	ToolkitCommands->MapAction(
+		Commands.PrimitiveQueryAndProbe,
+		FExecuteAction::CreateSP(this, &FPhysicsAssetEditor::OnSetPrimitiveCollision, ECollisionEnabled::QueryAndProbe),
+		FCanExecuteAction::CreateSP(this, &FPhysicsAssetEditor::CanSetPrimitiveCollision, ECollisionEnabled::QueryAndProbe),
+		FIsActionChecked::CreateSP(this, &FPhysicsAssetEditor::IsPrimitiveCollisionChecked, ECollisionEnabled::QueryAndProbe));
+
 	ToolkitCommands->MapAction(
 		Commands.PrimitiveQueryOnly,
 		FExecuteAction::CreateSP(this, &FPhysicsAssetEditor::OnSetPrimitiveCollision, ECollisionEnabled::QueryOnly),
@@ -827,6 +848,12 @@ void FPhysicsAssetEditor::BindCommands()
 		FCanExecuteAction::CreateSP(this, &FPhysicsAssetEditor::CanSetPrimitiveCollision, ECollisionEnabled::PhysicsOnly),
 		FIsActionChecked::CreateSP(this, &FPhysicsAssetEditor::IsPrimitiveCollisionChecked, ECollisionEnabled::PhysicsOnly));
 
+	ToolkitCommands->MapAction(
+		Commands.PrimitiveProbeOnly,
+		FExecuteAction::CreateSP(this, &FPhysicsAssetEditor::OnSetPrimitiveCollision, ECollisionEnabled::ProbeOnly),
+		FCanExecuteAction::CreateSP(this, &FPhysicsAssetEditor::CanSetPrimitiveCollision, ECollisionEnabled::ProbeOnly),
+		FIsActionChecked::CreateSP(this, &FPhysicsAssetEditor::IsPrimitiveCollisionChecked, ECollisionEnabled::ProbeOnly));
+
 	ToolkitCommands->MapAction(
 		Commands.PrimitiveNoCollision,
 		FExecuteAction::CreateSP(this, &FPhysicsAssetEditor::OnSetPrimitiveCollision, ECollisionEnabled::NoCollision),
@@ -1005,6 +1032,31 @@ void FPhysicsAssetEditor::BindCommands()
 		FExecuteAction::CreateSP(this, &FPhysicsAssetEditor::OnSelectKinematicBodies),
 		FCanExecuteAction::CreateSP(this, &FPhysicsAssetEditor::IsNotSimulation));
 
+	ToolkitCommands->MapAction(
+		Commands.SelectShapesQueryOnly,
+		FExecuteAction::CreateSP(this, &FPhysicsAssetEditor::OnSelectShapes, ECollisionEnabled::QueryOnly),
+		FCanExecuteAction::CreateSP(this, &FPhysicsAssetEditor::IsNotSimulation));
+
+	ToolkitCommands->MapAction(
+		Commands.SelectShapesQueryAndPhysics,
+		FExecuteAction::CreateSP(this, &FPhysicsAssetEditor::OnSelectShapes, ECollisionEnabled::QueryAndPhysics),
+		FCanExecuteAction::CreateSP(this, &FPhysicsAssetEditor::IsNotSimulation));
+
+	ToolkitCommands->MapAction(
+		Commands.SelectShapesPhysicsOnly,
+		FExecuteAction::CreateSP(this, &FPhysicsAssetEditor::OnSelectShapes, ECollisionEnabled::PhysicsOnly),
+		FCanExecuteAction::CreateSP(this, &FPhysicsAssetEditor::IsNotSimulation));
+
+	ToolkitCommands->MapAction(
+		Commands.SelectShapesQueryAndProbe,
+		FExecuteAction::CreateSP(this, &FPhysicsAssetEditor::OnSelectShapes, ECollisionEnabled::QueryAndProbe),
+		FCanExecuteAction::CreateSP(this, &FPhysicsAssetEditor::IsNotSimulation));
+
+	ToolkitCommands->MapAction(
+		Commands.SelectShapesProbeOnly,
+		FExecuteAction::CreateSP(this, &FPhysicsAssetEditor::OnSelectShapes, ECollisionEnabled::ProbeOnly),
+		FCanExecuteAction::CreateSP(this, &FPhysicsAssetEditor::IsNotSimulation));
+
 	ToolkitCommands->MapAction(
 		Commands.SelectAllConstraints,
 		FExecuteAction::CreateSP(this, &FPhysicsAssetEditor::OnSelectAllConstraints),
@@ -1362,10 +1414,12 @@ void FPhysicsAssetEditor::BuildMenuWidgetBody(FMenuBuilder& InMenuBuilder)
 				InSubMenuBuilder.AddMenuEntry( PhysicsAssetEditorCommands.DisableCollisionAll );
 				InSubMenuBuilder.EndSection();
 
-				InSubMenuBuilder.BeginSection("CollisionFilteringHeader", LOCTEXT("CollisionFilteringHeader", "CollisionFiltering"));
+				InSubMenuBuilder.BeginSection("CollisionFilteringHeader", LOCTEXT("CollisionFilteringHeader", "Collision Filtering"));
 				InSubMenuBuilder.AddMenuEntry(PhysicsAssetEditorCommands.PrimitiveQueryAndPhysics);
 				InSubMenuBuilder.AddMenuEntry(PhysicsAssetEditorCommands.PrimitiveQueryOnly);
+				InSubMenuBuilder.AddMenuEntry(PhysicsAssetEditorCommands.PrimitiveQueryAndProbe);
 				InSubMenuBuilder.AddMenuEntry(PhysicsAssetEditorCommands.PrimitivePhysicsOnly);
+				InSubMenuBuilder.AddMenuEntry(PhysicsAssetEditorCommands.PrimitiveProbeOnly);
 				InSubMenuBuilder.AddMenuEntry(PhysicsAssetEditorCommands.PrimitiveNoCollision);
 				InSubMenuBuilder.EndSection();
 
@@ -1514,6 +1568,11 @@ void FPhysicsAssetEditor::BuildMenuWidgetSelection(FMenuBuilder& InMenuBuilder)
 		InMenuBuilder.AddMenuEntry( Commands.ToggleShowOnlyConstrained );
 		InMenuBuilder.AddMenuEntry( Commands.ShowAll );
 		InMenuBuilder.AddMenuEntry( Commands.HideAll );
+		InMenuBuilder.AddMenuEntry( Commands.SelectShapesQueryOnly);
+		InMenuBuilder.AddMenuEntry( Commands.SelectShapesQueryAndPhysics);
+		InMenuBuilder.AddMenuEntry( Commands.SelectShapesPhysicsOnly);
+		InMenuBuilder.AddMenuEntry( Commands.SelectShapesQueryAndProbe);
+		InMenuBuilder.AddMenuEntry( Commands.SelectShapesProbeOnly);
 		InMenuBuilder.EndSection();
 	}
 	InMenuBuilder.PopCommandList();
@@ -2034,7 +2093,7 @@ void FPhysicsAssetEditor::ResetBoneCollision()
 
 		//deselect first
 		SharedData->ClearSelectedBody();
-		SharedData->SetSelectedBodiesAnyPrim(BodyIndices, true);
+		SharedData->SetSelectedBodiesAnyPrimitive(BodyIndices, true);
 	}
 	else
 	{
@@ -3027,7 +3086,7 @@ void FPhysicsAssetEditor::OnSelectAllBodies()
 
 	//first deselect everything
 	SharedData->ClearSelectedBody();
-	SharedData->SetSelectedBodiesAllPrim(NewSelectedBodies, true);
+	SharedData->SetSelectedBodiesAllPrimitive(NewSelectedBodies, true);
 
 }
 
@@ -3061,7 +3120,19 @@ void FPhysicsAssetEditor::OnSelectBodies(EPhysicsType PhysicsType)
 
 	//first deselect everything
 	SharedData->ClearSelectedBody();
-	SharedData->SetSelectedBodiesAllPrim(NewSelectedBodies, true);
+	SharedData->SetSelectedBodiesAllPrimitive(NewSelectedBodies, true);
+}
+
+void FPhysicsAssetEditor::OnSelectShapes(const ECollisionEnabled::Type CollisionEnabled)
+{
+	UPhysicsAsset * const PhysicsAsset = SharedData->EditorSkelComp->GetPhysicsAsset();
+	TSet<int32> SelectedBodyIndices;
+	for (const FPhysicsAssetEditorSharedData::FSelection& SelectedBody : SharedData->SelectedBodies)
+	{
+		SelectedBodyIndices.Add(SelectedBody.Index);
+	}
+	SharedData->ClearSelectedBody();
+	SharedData->SetSelectedBodiesPrimitivesWithCollisionType(SelectedBodyIndices.Array(), CollisionEnabled, true);
 }
 
 void FPhysicsAssetEditor::OnSelectAllConstraints()
@@ -3241,7 +3312,7 @@ void FPhysicsAssetEditor::HandleGraphObjectsSelected(const TArrayView<UObject*>&
 			}
 		}
 
-		SharedData->SetSelectedBodiesAnyPrim(SelectedBodyIndices, true);
+		SharedData->SetSelectedBodiesAnyPrimitive(SelectedBodyIndices, true);
 		SharedData->SetSelectedConstraints(SelectedConstraintIndices, true);
 
 		SkeletonTree->SelectItemsBy([&SelectedBodySetups, &SelectedConstraintTemplates](const TSharedRef<ISkeletonTreeItem>& InItem, bool& bInOutExpand)
@@ -3331,7 +3402,7 @@ void FPhysicsAssetEditor::HandleSelectionChanged(const TArrayView<TSharedPtr<ISk
 			}
 
 			SharedData->SetSelectedBodies(SelectedBodies, true);
-			SharedData->SetSelectedBodiesAnyPrim(SelectedBodiesAnyPrim, true);
+			SharedData->SetSelectedBodiesAnyPrimitive(SelectedBodiesAnyPrim, true);
 			SharedData->SetSelectedConstraints(SelectedConstraints, true);
 
 			if(!bBoneSelected)
diff --git a/Engine/Source/Editor/PhysicsAssetEditor/Private/PhysicsAssetEditor.h b/Engine/Source/Editor/PhysicsAssetEditor/Private/PhysicsAssetEditor.h
--- a/Engine/Source/Editor/PhysicsAssetEditor/Private/PhysicsAssetEditor.h
+++ b/Engine/Source/Editor/PhysicsAssetEditor/Private/PhysicsAssetEditor.h
@@ -298,6 +298,7 @@ private:
 	void OnSelectKinematicBodies();
 	void OnSelectSimulatedBodies();
 	void OnSelectBodies(EPhysicsType PhysicsType = EPhysicsType::PhysType_Simulated);
+	void OnSelectShapes(const ECollisionEnabled::Type CollisionEnabled);
 	void OnSelectAllConstraints();
 	void OnToggleSelectionType(bool bIgnoreUserConstraints);
 	void OnToggleShowSelected();
diff --git a/Engine/Source/Editor/PhysicsAssetEditor/Private/PhysicsAssetEditorActions.cpp b/Engine/Source/Editor/PhysicsAssetEditor/Private/PhysicsAssetEditorActions.cpp
--- a/Engine/Source/Editor/PhysicsAssetEditor/Private/PhysicsAssetEditorActions.cpp
+++ b/Engine/Source/Editor/PhysicsAssetEditor/Private/PhysicsAssetEditorActions.cpp
@@ -58,9 +58,11 @@ void FPhysicsAssetEditorCommands::RegisterCommands()
 	UI_COMMAND(AddBox, "Add Box", "Add Box To Selected Bone", EUserInterfaceActionType::Button, FInputChord(EModifierKey::Control, EKeys::Three));
 	UI_COMMAND(AddTaperedCapsule, "Add Tapered Capsule (Clothing Only)", "Add Tapered Capsule To Selected Bone. This is only used by clothing, so will have no effect on rigid body collision, overlaps or bounds.", EUserInterfaceActionType::Button, FInputChord(EModifierKey::Control, EKeys::Four));
 	UI_COMMAND(DeletePrimitive, "Delete", "Delete Selected Primitive(s)", EUserInterfaceActionType::Button, FInputChord());
-	UI_COMMAND(PrimitiveQueryAndPhysics, "All Collision", "Enable all collision on selected primitive(s)", EUserInterfaceActionType::RadioButton, FInputChord());
-	UI_COMMAND(PrimitiveQueryOnly, "Query Collision Only", "Enable query collision only on selected primitive(s)", EUserInterfaceActionType::RadioButton, FInputChord());
-	UI_COMMAND(PrimitivePhysicsOnly, "Physics Collision Only", "Enable physics collision only on selected primitive(s)", EUserInterfaceActionType::RadioButton, FInputChord());
+	UI_COMMAND(PrimitiveQueryAndPhysics, "Query and Physics", "Enable all collision on selected primitive(s)", EUserInterfaceActionType::RadioButton, FInputChord());
+	UI_COMMAND(PrimitiveQueryAndProbe, "Query and Probe", "Enable query and probe collision on selected primitive(s)", EUserInterfaceActionType::RadioButton, FInputChord());
+	UI_COMMAND(PrimitiveQueryOnly, "Query Only", "Enable query collision only on selected primitive(s)", EUserInterfaceActionType::RadioButton, FInputChord());
+	UI_COMMAND(PrimitivePhysicsOnly, "Physics Only", "Enable physics collision only on selected primitive(s)", EUserInterfaceActionType::RadioButton, FInputChord());
+	UI_COMMAND(PrimitiveProbeOnly, "Probe Only", "Enable probe collision only on selected primitive(s)", EUserInterfaceActionType::RadioButton, FInputChord());
 	UI_COMMAND(PrimitiveNoCollision, "No Collision", "Disable all collision on Selected primitive(s)", EUserInterfaceActionType::RadioButton, FInputChord());
 	UI_COMMAND(PrimitiveContributeToMass, "Primitive Contributes To Mass", "Toggle the contribution of selected primitive's volume to the overall mass of the body", EUserInterfaceActionType::Check, FInputChord());
 	UI_COMMAND(DuplicatePrimitive, "Duplicate", "Duplicate Selected Primitive(s)", EUserInterfaceActionType::Button, FInputChord());
@@ -84,6 +86,11 @@ void FPhysicsAssetEditorCommands::RegisterCommands()
 	UI_COMMAND(SelectAllBodies, "Select All Bodies", "Select All Bodies", EUserInterfaceActionType::Button, FInputChord(EModifierKey::Control, EKeys::A));
 	UI_COMMAND(SelectSimulatedBodies, "Select Simulated Bodies", "Select Simulated Bodies", EUserInterfaceActionType::Button, FInputChord(EModifierKey::Control, EKeys::J));
 	UI_COMMAND(SelectKinematicBodies, "Select Kinematic Bodies", "Select Kinematic Bodies", EUserInterfaceActionType::Button, FInputChord(EModifierKey::Control, EKeys::K));
+	UI_COMMAND(SelectShapesQueryOnly, "Select QueryOnly Shapes", "Select Kinematic Bodies", EUserInterfaceActionType::Button, FInputChord());
+	UI_COMMAND(SelectShapesQueryAndPhysics, "Select QueryAndPhysics Shapes", "Select Kinematic Bodies", EUserInterfaceActionType::Button, FInputChord());
+	UI_COMMAND(SelectShapesPhysicsOnly, "Select PhysicsOnly Shapes", "Select Kinematic Bodies", EUserInterfaceActionType::Button, FInputChord());
+	UI_COMMAND(SelectShapesQueryAndProbe, "Select QueryAndProbe Shapes", "Select Kinematic Bodies", EUserInterfaceActionType::Button, FInputChord());
+	UI_COMMAND(SelectShapesProbeOnly, "Select ProbeOnly Shapes", "Select Kinematic Bodies", EUserInterfaceActionType::Button, FInputChord());
 	UI_COMMAND(SelectAllConstraints, "Select All Constraints", "Select All Constraints", EUserInterfaceActionType::Button, FInputChord(EModifierKey::Control | EModifierKey::Shift, EKeys::A));
 	UI_COMMAND(ToggleSelectionType, "Toggle Selection Type", "Select bodies from constraints, or constraints from bodies, depending on selection.", EUserInterfaceActionType::Button, FInputChord(EModifierKey::Control, EKeys::T));
 	UI_COMMAND(ToggleSelectionTypeWithUserConstraints, "Toggle Selection Type (With User Constraints)", "Select bodies from constraints, or constraints from bodies, depending on selection (account for user constraints).", EUserInterfaceActionType::Button, FInputChord(EModifierKey::Control | EModifierKey::Shift, EKeys::T));
diff --git a/Engine/Source/Editor/PhysicsAssetEditor/Private/PhysicsAssetEditorActions.h b/Engine/Source/Editor/PhysicsAssetEditor/Private/PhysicsAssetEditorActions.h
--- a/Engine/Source/Editor/PhysicsAssetEditor/Private/PhysicsAssetEditorActions.h
+++ b/Engine/Source/Editor/PhysicsAssetEditor/Private/PhysicsAssetEditorActions.h
@@ -65,8 +65,10 @@ public:
 	TSharedPtr<FUICommandInfo> EnableCollision;
 	TSharedPtr<FUICommandInfo> EnableCollisionAll;
 	TSharedPtr<FUICommandInfo> PrimitiveQueryAndPhysics;
+	TSharedPtr<FUICommandInfo> PrimitiveQueryAndProbe;
 	TSharedPtr<FUICommandInfo> PrimitiveQueryOnly;
 	TSharedPtr<FUICommandInfo> PrimitivePhysicsOnly;
+	TSharedPtr<FUICommandInfo> PrimitiveProbeOnly;
 	TSharedPtr<FUICommandInfo> PrimitiveNoCollision;
 	TSharedPtr<FUICommandInfo> PrimitiveContributeToMass;
 	TSharedPtr<FUICommandInfo> WeldToBody;
@@ -96,6 +98,11 @@ public:
 	TSharedPtr<FUICommandInfo> SelectAllBodies;
 	TSharedPtr<FUICommandInfo> SelectSimulatedBodies;
 	TSharedPtr<FUICommandInfo> SelectKinematicBodies;
+	TSharedPtr<FUICommandInfo> SelectShapesQueryOnly;
+	TSharedPtr<FUICommandInfo> SelectShapesQueryAndPhysics;
+	TSharedPtr<FUICommandInfo> SelectShapesPhysicsOnly;
+	TSharedPtr<FUICommandInfo> SelectShapesQueryAndProbe;
+	TSharedPtr<FUICommandInfo> SelectShapesProbeOnly;
 	TSharedPtr<FUICommandInfo> SelectAllConstraints;
 	TSharedPtr<FUICommandInfo> ToggleSelectionType;
 	TSharedPtr<FUICommandInfo> ToggleSelectionTypeWithUserConstraints;
diff --git a/Engine/Source/Editor/PhysicsAssetEditor/Private/PhysicsAssetEditorSharedData.cpp b/Engine/Source/Editor/PhysicsAssetEditor/Private/PhysicsAssetEditorSharedData.cpp
--- a/Engine/Source/Editor/PhysicsAssetEditor/Private/PhysicsAssetEditorSharedData.cpp
+++ b/Engine/Source/Editor/PhysicsAssetEditor/Private/PhysicsAssetEditorSharedData.cpp
@@ -518,63 +518,55 @@ void FPhysicsAssetEditorSharedData::RefreshPhysicsAssetChange(const UPhysicsAsse
 	}
 }
 
-void FPhysicsAssetEditorSharedData::SetSelectedBodyAnyPrim(int32 BodyIndex, bool bSelected)
+void FPhysicsAssetEditorSharedData::SetSelectedBodyAnyPrimitive(int32 BodyIndex, bool bSelected)
 {
-	SetSelectedBodiesAnyPrim({ BodyIndex }, bSelected);
+	SetSelectedBodiesAnyPrimitive({ BodyIndex }, bSelected);
 }
 
-void FPhysicsAssetEditorSharedData::SetSelectedBodiesAnyPrim(const TArray<int32>& BodiesIndices, bool bSelected)
+void FPhysicsAssetEditorSharedData::SetSelectedBodiesAnyPrimitive(const TArray<int32>& BodiesIndices, bool bSelected)
 {
-	if (BodiesIndices.Num() == 0)
+	SetSelectedBodiesPrimitives(BodiesIndices, bSelected, [](const TArray<FSelection>& CurrentSelection, const FKShapeElem& Primitive)
 	{
-		return;
-	}
+		// Select the first primitive
+		return CurrentSelection.Num() == 0;
+	});
+}
 
-	if (BodiesIndices.Num() == 1 && BodiesIndices[0] == INDEX_NONE)
+void FPhysicsAssetEditorSharedData::SetSelectedBodiesAllPrimitive(const TArray<int32>& BodiesIndices, bool bSelected)
+{
+	SetSelectedBodiesPrimitives(BodiesIndices, bSelected, [](const TArray<FSelection>& CurrentSelection, const FKShapeElem& Primitive)
 	{
-		ClearSelectedBody();
-		return;
-	}
+		// Select all primitives
+		return true;
+	});
+}
 
-	TArray<FSelection> NewSelection;
-	for (const int32 BodyIndex : BodiesIndices)
+void FPhysicsAssetEditorSharedData::SetSelectedBodiesPrimitivesWithCollisionType(const TArray<int32>& BodiesIndices, const ECollisionEnabled::Type CollisionType, bool bSelected)
+{
+	SetSelectedBodiesPrimitives(BodiesIndices, bSelected, [CollisionType](const TArray<FSelection>& CurrentSelection, const FKShapeElem& Primitive)
 	{
-		UBodySetup* BodySetup = PhysicsAsset->SkeletalBodySetups[BodyIndex];
-		check(BodySetup);
+		// Select primitives which match the collision type
+		return Primitive.GetCollisionEnabled() == CollisionType;
+	});
+}
 
-		if (BodySetup->AggGeom.SphereElems.Num() > 0)
-		{
-			NewSelection.Add(FSelection(BodyIndex, EAggCollisionShape::Sphere, 0));
-		}
-		else if (BodySetup->AggGeom.BoxElems.Num() > 0)
-		{
-			NewSelection.Add(FSelection(BodyIndex, EAggCollisionShape::Box, 0));
-		}
-		else if (BodySetup->AggGeom.SphylElems.Num() > 0)
-		{
-			NewSelection.Add(FSelection(BodyIndex, EAggCollisionShape::Sphyl, 0));
-		}
-		else if (BodySetup->AggGeom.ConvexElems.Num() > 0)
-		{
-			NewSelection.Add(FSelection(BodyIndex, EAggCollisionShape::Convex, 0));
-		}
-		else if (BodySetup->AggGeom.TaperedCapsuleElems.Num() > 0)
-		{
-			NewSelection.Add(FSelection(BodyIndex, EAggCollisionShape::TaperedCapsule, 0));
-		}
-		else
+namespace
+{
+	template <typename TShapeElem>
+	void SetSelectedBodiesPrimitivesHelper(const int32 BodyIndex, const TArray<TShapeElem>& ShapeElems, TArray<FPhysicsAssetEditorSharedData::FSelection>& SelectedElems, const TFunction<bool(const TArray<FPhysicsAssetEditorSharedData::FSelection>&, const FKShapeElem&)>& Predicate)
+	{
+		for (int32 PrimitiveIndex = 0; PrimitiveIndex < ShapeElems.Num(); ++PrimitiveIndex)
 		{
-			UE_LOG(LogPhysicsAssetEditor, Fatal, TEXT("Body Setup with index %d has No Primitives!"), BodyIndex);
+			const TShapeElem& ShapeElem = ShapeElems[PrimitiveIndex];
+			if (Predicate(SelectedElems, ShapeElem))
+			{
+				SelectedElems.Add(FPhysicsAssetEditorSharedData::FSelection(BodyIndex, ShapeElem.GetShapeType(), PrimitiveIndex));
+			}
 		}
 	}
-
-	if (NewSelection.Num() > 0)
-	{
-		SetSelectedBodies(NewSelection, bSelected);
-	}
 }
 
-void FPhysicsAssetEditorSharedData::SetSelectedBodiesAllPrim(const TArray<int32>& BodiesIndices, bool bSelected)
+void FPhysicsAssetEditorSharedData::SetSelectedBodiesPrimitives(const TArray<int32>& BodiesIndices, bool bSelected, const TFunction<bool(const TArray<FSelection>&, const FKShapeElem&)>& Predicate)
 {
 	if (BodiesIndices.Num() == 0)
 	{
@@ -593,35 +585,15 @@ void FPhysicsAssetEditorSharedData::SetSelectedBodiesAllPrim(const TArray<int32>
 		UBodySetup* BodySetup = PhysicsAsset->SkeletalBodySetups[BodyIndex];
 		check(BodySetup);
 
-		for (int32 PrimitiveIndex = 0; PrimitiveIndex < BodySetup->AggGeom.SphereElems.Num(); ++PrimitiveIndex)
-		{
-			NewSelection.Add(FSelection(BodyIndex, EAggCollisionShape::Sphere, PrimitiveIndex));
-		}
-
-		for (int32 PrimitiveIndex = 0; PrimitiveIndex < BodySetup->AggGeom.BoxElems.Num(); ++PrimitiveIndex)
-		{
-			NewSelection.Add(FSelection(BodyIndex, EAggCollisionShape::Box, PrimitiveIndex));
-		}
-
-		for (int32 PrimitiveIndex = 0; PrimitiveIndex < BodySetup->AggGeom.SphylElems.Num(); ++PrimitiveIndex)
-		{
-			NewSelection.Add(FSelection(BodyIndex, EAggCollisionShape::Sphyl, PrimitiveIndex));
-		}
-
-		for (int32 PrimitiveIndex = 0; PrimitiveIndex < BodySetup->AggGeom.ConvexElems.Num(); ++PrimitiveIndex)
-		{
-			NewSelection.Add(FSelection(BodyIndex, EAggCollisionShape::Convex, PrimitiveIndex));
-		}
-		for (int32 PrimitiveIndex = 0; PrimitiveIndex < BodySetup->AggGeom.TaperedCapsuleElems.Num(); ++PrimitiveIndex)
-		{
-			NewSelection.Add(FSelection(BodyIndex, EAggCollisionShape::TaperedCapsule, PrimitiveIndex));
-		}
+		const FKAggregateGeom& AggGeom = BodySetup->AggGeom;
+		SetSelectedBodiesPrimitivesHelper(BodyIndex, AggGeom.SphereElems, NewSelection, Predicate);
+		SetSelectedBodiesPrimitivesHelper(BodyIndex, AggGeom.BoxElems, NewSelection, Predicate);
+		SetSelectedBodiesPrimitivesHelper(BodyIndex, AggGeom.SphylElems, NewSelection, Predicate);
+		SetSelectedBodiesPrimitivesHelper(BodyIndex, AggGeom.ConvexElems, NewSelection, Predicate);
+		SetSelectedBodiesPrimitivesHelper(BodyIndex, AggGeom.TaperedCapsuleElems, NewSelection, Predicate);
 	}
 
-	if (NewSelection.Num() > 0)
-	{
-		SetSelectedBodies(NewSelection, bSelected);
-	}
+	SetSelectedBodies(NewSelection, bSelected);
 }
 
 void FPhysicsAssetEditorSharedData::ClearSelectedBody()
@@ -725,7 +697,7 @@ void FPhysicsAssetEditorSharedData::ToggleSelectionType(bool bIgnoreUserConstrai
 	ClearSelectedBody();
 	ClearSelectedConstraints();
 
-	SetSelectedBodiesAllPrim(NewSelectedBodies.Array(), true);
+	SetSelectedBodiesAllPrimitive(NewSelectedBodies.Array(), true);
 	SetSelectedConstraints(NewSelectedConstraints.Array(), true);
 }
 
@@ -2019,7 +1991,7 @@ void FPhysicsAssetEditorSharedData::MakeNewBody(int32 NewBoneIndex, bool bAutoSe
 
 	if (bAutoSelect)
 	{
-		SetSelectedBodyAnyPrim(NewBodyIndex, true);
+		SetSelectedBodyAnyPrimitive(NewBodyIndex, true);
 	}
 	
 
diff --git a/Engine/Source/Editor/PhysicsAssetEditor/Private/PhysicsAssetEditorSharedData.h b/Engine/Source/Editor/PhysicsAssetEditor/Private/PhysicsAssetEditorSharedData.h
--- a/Engine/Source/Editor/PhysicsAssetEditor/Private/PhysicsAssetEditorSharedData.h
+++ b/Engine/Source/Editor/PhysicsAssetEditor/Private/PhysicsAssetEditorSharedData.h
@@ -160,9 +160,11 @@ public:
 	void ToggleShowOnlySelected();
 	void ShowSelected();
 	void HideSelected();
-	void SetSelectedBodyAnyPrim(int32 BodyIndex, bool bSelected);
-	void SetSelectedBodiesAnyPrim(const TArray<int32>& BodiesIndices, bool bSelected);
-	void SetSelectedBodiesAllPrim(const TArray<int32>& BodiesIndices, bool bSelected);
+	void SetSelectedBodyAnyPrimitive(int32 BodyIndex, bool bSelected);
+	void SetSelectedBodiesAnyPrimitive(const TArray<int32>& BodiesIndices, bool bSelected);
+	void SetSelectedBodiesAllPrimitive(const TArray<int32>& BodiesIndices, bool bSelected);
+	void SetSelectedBodiesPrimitivesWithCollisionType(const TArray<int32>& BodiesIndices, const ECollisionEnabled::Type CollisionType, bool bSelected);
+	void SetSelectedBodiesPrimitives(const TArray<int32>& BodiesIndices, bool bSelected, const TFunction<bool(const TArray<FSelection>&, const FKShapeElem&)>& Predicate);
 	void DeleteCurrentPrim();
 	void DeleteBody(int32 DelBodyIndex, bool bRefreshComponent=true);
 	void RefreshPhysicsAssetChange(const UPhysicsAsset* InPhysAsset, bool bFullClothRefresh = true);
