[UE][FIX] Multi-Select in PhAT tree view
SetSelectedBodiesAnyPrimitive now selects the first primitive in each body in the supplied list, not just the first body.

#rmx

diff --git a/Engine/Source/Editor/PhysicsAssetEditor/Private/PhysicsAssetEditorSharedData.cpp b/Engine/Source/Editor/PhysicsAssetEditor/Private/PhysicsAssetEditorSharedData.cpp
--- a/Engine/Source/Editor/PhysicsAssetEditor/Private/PhysicsAssetEditorSharedData.cpp
+++ b/Engine/Source/Editor/PhysicsAssetEditor/Private/PhysicsAssetEditorSharedData.cpp
@@ -525,16 +525,16 @@ void FPhysicsAssetEditorSharedData::SetSelectedBodyAnyPrimitive(int32 BodyIndex,
 
 void FPhysicsAssetEditorSharedData::SetSelectedBodiesAnyPrimitive(const TArray<int32>& BodiesIndices, bool bSelected)
 {
-	SetSelectedBodiesPrimitives(BodiesIndices, bSelected, [](const TArray<FSelection>& CurrentSelection, const FKShapeElem& Primitive)
+	SetSelectedBodiesPrimitives(BodiesIndices, bSelected, [](const TArray<FSelection>& CurrentSelection, const int32 BodyIndex, const FKShapeElem& Primitive)
 	{
-		// Select the first primitive
-		return CurrentSelection.Num() == 0;
+		// Select the first primitive in the supplied body. 
+		return !CurrentSelection.ContainsByPredicate([BodyIndex](const FSelection& InSelection) {return InSelection.Index == BodyIndex; }); // Predicate returns true if the supplied body index is not already in the current selection.
 	});
 }
 
 void FPhysicsAssetEditorSharedData::SetSelectedBodiesAllPrimitive(const TArray<int32>& BodiesIndices, bool bSelected)
 {
-	SetSelectedBodiesPrimitives(BodiesIndices, bSelected, [](const TArray<FSelection>& CurrentSelection, const FKShapeElem& Primitive)
+	SetSelectedBodiesPrimitives(BodiesIndices, bSelected, [](const TArray<FSelection>& CurrentSelection, const int32 BodyIndex, const FKShapeElem& Primitive)
 	{
 		// Select all primitives
 		return true;
@@ -543,7 +543,7 @@ void FPhysicsAssetEditorSharedData::SetSelectedBodiesAllPrimitive(const TArray<i
 
 void FPhysicsAssetEditorSharedData::SetSelectedBodiesPrimitivesWithCollisionType(const TArray<int32>& BodiesIndices, const ECollisionEnabled::Type CollisionType, bool bSelected)
 {
-	SetSelectedBodiesPrimitives(BodiesIndices, bSelected, [CollisionType](const TArray<FSelection>& CurrentSelection, const FKShapeElem& Primitive)
+	SetSelectedBodiesPrimitives(BodiesIndices, bSelected, [CollisionType](const TArray<FSelection>& CurrentSelection, const int32 BodyIndex, const FKShapeElem& Primitive)
 	{
 		// Select primitives which match the collision type
 		return Primitive.GetCollisionEnabled() == CollisionType;
@@ -553,12 +553,12 @@ void FPhysicsAssetEditorSharedData::SetSelectedBodiesPrimitivesWithCollisionType
 namespace
 {
 	template <typename TShapeElem>
-	void SetSelectedBodiesPrimitivesHelper(const int32 BodyIndex, const TArray<TShapeElem>& ShapeElems, TArray<FPhysicsAssetEditorSharedData::FSelection>& SelectedElems, const TFunction<bool(const TArray<FPhysicsAssetEditorSharedData::FSelection>&, const FKShapeElem&)>& Predicate)
+	void SetSelectedBodiesPrimitivesHelper(const int32 BodyIndex, const TArray<TShapeElem>& ShapeElems, TArray<FPhysicsAssetEditorSharedData::FSelection>& SelectedElems, const TFunction<bool(const TArray<FPhysicsAssetEditorSharedData::FSelection>&, const int32 BodyIndex, const FKShapeElem&)>& Predicate)
 	{
 		for (int32 PrimitiveIndex = 0; PrimitiveIndex < ShapeElems.Num(); ++PrimitiveIndex)
 		{
 			const TShapeElem& ShapeElem = ShapeElems[PrimitiveIndex];
-			if (Predicate(SelectedElems, ShapeElem))
+			if (Predicate(SelectedElems, BodyIndex, ShapeElem))
 			{
 				SelectedElems.Add(FPhysicsAssetEditorSharedData::FSelection(BodyIndex, ShapeElem.GetShapeType(), PrimitiveIndex));
 			}
@@ -566,7 +566,7 @@ namespace
 	}
 }
 
-void FPhysicsAssetEditorSharedData::SetSelectedBodiesPrimitives(const TArray<int32>& BodiesIndices, bool bSelected, const TFunction<bool(const TArray<FSelection>&, const FKShapeElem&)>& Predicate)
+void FPhysicsAssetEditorSharedData::SetSelectedBodiesPrimitives(const TArray<int32>& BodiesIndices, bool bSelected, const TFunction<bool(const TArray<FSelection>&, const int32 BodyIndex, const FKShapeElem&)>& Predicate)
 {
 	if (BodiesIndices.Num() == 0)
 	{
diff --git a/Engine/Source/Editor/PhysicsAssetEditor/Private/PhysicsAssetEditorSharedData.h b/Engine/Source/Editor/PhysicsAssetEditor/Private/PhysicsAssetEditorSharedData.h
--- a/Engine/Source/Editor/PhysicsAssetEditor/Private/PhysicsAssetEditorSharedData.h
+++ b/Engine/Source/Editor/PhysicsAssetEditor/Private/PhysicsAssetEditorSharedData.h
@@ -164,7 +164,7 @@ public:
 	void SetSelectedBodiesAnyPrimitive(const TArray<int32>& BodiesIndices, bool bSelected);
 	void SetSelectedBodiesAllPrimitive(const TArray<int32>& BodiesIndices, bool bSelected);
 	void SetSelectedBodiesPrimitivesWithCollisionType(const TArray<int32>& BodiesIndices, const ECollisionEnabled::Type CollisionType, bool bSelected);
-	void SetSelectedBodiesPrimitives(const TArray<int32>& BodiesIndices, bool bSelected, const TFunction<bool(const TArray<FSelection>&, const FKShapeElem&)>& Predicate);
+	void SetSelectedBodiesPrimitives(const TArray<int32>& BodiesIndices, bool bSelected, const TFunction<bool(const TArray<FSelection>&, const int32 BodyIndex, const FKShapeElem&)>& Predicate);
 	void DeleteCurrentPrim();
 	void DeleteBody(int32 DelBodyIndex, bool bRefreshComponent=true);
 	void RefreshPhysicsAssetChange(const UPhysicsAsset* InPhysAsset, bool bFullClothRefresh = true);
