[UE][FIX] Fix for Phat M to mirror bodies breaks the collision setups
- Extended bone mirroring to support bone names that have "_l_" or "_r_" patterns in the middle (previously just supported l_/r_ prefix or _l/_r suffix)
- Transfer collision interaction markup from source to mirrored bones by mirroring the bones controlled by the source body's colliding bodies.
- Added a notification to report any errors in mirroring collision interactions

diff --git a/Engine/Source/Editor/PhysicsAssetEditor/Private/PhysicsAssetEditorSharedData.cpp b/Engine/Source/Editor/PhysicsAssetEditor/Private/PhysicsAssetEditorSharedData.cpp
--- a/Engine/Source/Editor/PhysicsAssetEditor/Private/PhysicsAssetEditorSharedData.cpp
+++ b/Engine/Source/Editor/PhysicsAssetEditor/Private/PhysicsAssetEditorSharedData.cpp
@@ -34,6 +34,8 @@
 #include "PropertyEditorModule.h"
 #include "Widgets/Input/SButton.h"
 #include "Widgets/Text/STextBlock.h"
+#include "Widgets/Notifications/SNotificationList.h"
+#include "Framework/Notifications/NotificationManager.h"
 #include "ClothingSimulationInteractor.h"
 #include "UnrealExporter.h"
 #include "Exporters/Exporter.h"
@@ -323,6 +325,7 @@ struct FMirrorInfo
 	int32 BoneIndex;
 	int32 BodyIndex;
 	int32 ConstraintIndex;
+	TArray<FName> CollidingBodyBoneNames; // Names of the controlling bones of all bodies that this body can collide with.
 	FMirrorInfo()
 	{
 		BoneIndex = INDEX_NONE;
@@ -337,21 +340,37 @@ void FPhysicsAssetEditorSharedData::Mirror()
 	USkeletalMesh* EditorSkelMesh = PhysicsAsset->GetPreviewMesh();
 	if(EditorSkelMesh)
 	{
+		TArray<int32> DBG_SrcCollidingBodies;
 
+		// Build list of all bodies and constraints to be mirrored
 		TArray<FMirrorInfo> MirrorInfos;
+		MirrorInfos.Reserve(SelectedBodies.Num() + SelectedConstraints.Num());
 
 		for (const FSelection& Selection : SelectedBodies)
 		{
-			MirrorInfos.AddUninitialized();
+			MirrorInfos.AddDefaulted();
 			FMirrorInfo & MirrorInfo = MirrorInfos[MirrorInfos.Num() - 1];
 			MirrorInfo.BoneName = PhysicsAsset->SkeletalBodySetups[Selection.Index]->BoneName;
 			MirrorInfo.BodyIndex = Selection.Index;
 			MirrorInfo.ConstraintIndex = PhysicsAsset->FindConstraintIndex(MirrorInfo.BoneName);
+			
+			// Record all the colliding body bone names 
+			// - This must be done before the bodies are mirrored because information may be lost in that process (for example, a user 
+			//   could select a mirrored pair of bodies. Both would be destroyed and recreated before collision interactions were mirrored).
+			// - Need to store bone names as body indexs can change during mirroring.
+			for (int32 CollidingBodyIndex = 0; CollidingBodyIndex < PhysicsAsset->SkeletalBodySetups.Num(); ++CollidingBodyIndex)
+			{
+				if (PhysicsAsset->IsCollisionEnabled(CollidingBodyIndex, MirrorInfo.BodyIndex))
+				{
+					const FName CollidingBoneName = PhysicsAsset->SkeletalBodySetups[CollidingBodyIndex]->BoneName;
+					MirrorInfo.CollidingBodyBoneNames.Add(CollidingBoneName);
+				}
+			}
 		}
 
 		for (const FSelection& Selection : SelectedConstraints)
 		{
-			MirrorInfos.AddUninitialized();
+			MirrorInfos.AddDefaulted();
 			FMirrorInfo & MirrorInfo = MirrorInfos[MirrorInfos.Num() - 1];
 			MirrorInfo.BoneName = PhysicsAsset->ConstraintSetup[Selection.Index]->DefaultInstance.ConstraintBone1;
 			MirrorInfo.BodyIndex = PhysicsAsset->FindBodyIndex(MirrorInfo.BoneName);
@@ -386,7 +405,7 @@ void FPhysicsAssetEditorSharedData::Mirror()
 				for (FKBoxElem& Box : DestBody->AggGeom.BoxElems)
 				{
 					Box.Rotation	= (Box.Rotation.Quaternion()*ArtistMirrorConvention).Rotator();
-					Box.Center      = -Box.Center;
+					Box.Center		= -Box.Center;
 				}
 				for (FKSphereElem& Sphere : DestBody->AggGeom.SphereElems)
 				{
@@ -406,6 +425,83 @@ void FPhysicsAssetEditorSharedData::Mirror()
 				}
 			}
 		}
+
+		// Mirror collision interactions - Do this after all mirrored bodies have been created as there may be collision interactions between the new bodies.
+		{
+			FString MirrorCollisionsMissingBones;
+			FString MirrorCollisionsMissingBodies;
+			uint32 MissingBodyCount = 0;
+			uint32 MissingBoneCount = 0;
+
+			for (FMirrorInfo& MirrorInfo : MirrorInfos)
+			{
+				const int32 SourceBoneIndex = EditorSkelMesh->GetRefSkeleton().FindBoneIndex(MirrorInfo.BoneName);
+				const int32 MirrorBoneIndex = PhysicsAsset->FindMirroredBone(EditorSkelMesh, SourceBoneIndex);
+
+				if (MirrorBoneIndex != INDEX_NONE)
+				{
+					const int32 SourceBodyIndex = MirrorInfo.BodyIndex;
+
+					for(FName SourceCollidingBoneName : MirrorInfo.CollidingBodyBoneNames)
+					{
+						const int32 SourceCollidingBoneIndex = EditorSkelMesh->GetRefSkeleton().FindBoneIndex(SourceCollidingBoneName); // Find Index of the bone associated with the body that the source body was allowed to collide with.
+						const int32 MirrorCollidingBoneIndex = PhysicsAsset->FindMirroredBone(EditorSkelMesh, SourceCollidingBoneIndex); // Find the index of the bone that mirrors the colliding body's bone.
+						const FName MirrorCollidingBoneName = EditorSkelMesh->GetRefSkeleton().GetBoneName(MirrorCollidingBoneIndex); // Find the name of the bone that mirrors the colliding body's bone.
+						const int32 MirrorCollidingBodyIndex = PhysicsAsset->FindBodyIndex(MirrorCollidingBoneName); // Find the index of the colliding body.
+
+						if (MirrorCollidingBodyIndex != INDEX_NONE)
+						{
+							const FName MirrorBoneName = EditorSkelMesh->GetRefSkeleton().GetBoneName(MirrorBoneIndex);
+							const int32 MirrorBodyIndex = PhysicsAsset->FindBodyIndex(MirrorBoneName);
+
+							PhysicsAsset->EnableCollision(MirrorCollidingBodyIndex, MirrorBodyIndex); // Enable collisions with the body associated with that bone.
+						}
+						else // Error reporting
+						{
+							if (MirrorCollidingBoneIndex != INDEX_NONE) // Found the mirrored bone but failed to find an associated physics body
+							{
+								MirrorCollisionsMissingBodies += MirrorCollidingBoneName.ToString() + "\n";
+								++MissingBodyCount;
+							}
+							else // Failed to find the mirrored bone.
+							{
+								MirrorCollisionsMissingBones += SourceCollidingBoneName.ToString() + "\n";
+								++MissingBoneCount;
+							}
+						}
+					}
+				}
+
+				// Display an error notification if neccesary.
+				if (!(MirrorCollisionsMissingBones.IsEmpty() && MirrorCollisionsMissingBodies.IsEmpty()))
+				{
+					// Construct error message for failed collision mirroring.
+					FText MissingMirrorBodiesErrorText;
+					FText MissingMirrorBonesErrorText;
+
+					if (MissingBodyCount > 0)
+					{
+						MissingMirrorBodiesErrorText = FText::Format(LOCTEXT("MissingMirrorBody", "Missing {0}|plural(one=body,other=bodies) for {0}|plural(one=bone,other=bones):\n{1}"), MissingBodyCount, FText::FromString(MirrorCollisionsMissingBodies));
+					}
+
+					if (MissingBoneCount > 0)
+					{
+						MissingMirrorBonesErrorText = FText::Format(LOCTEXT("MissingMirrorBone", "Missing {0}|plural(one=mirror,other=mirrors) for {0}|plural(one=bone,other=bones):\n{1}Note: Mirroring is based entirely on bone name matching."), MissingBoneCount, FText::FromString(MirrorCollisionsMissingBones));
+					}
+
+					const FText ErrorMsg = FText::Format(LOCTEXT("FailedToMirrorCollisions", "Failed to mirror all collisions\n{0}{1}"), MissingMirrorBodiesErrorText, MissingMirrorBonesErrorText);
+
+					// Display notification.
+					FNotificationInfo Info(ErrorMsg);
+					Info.ExpireDuration = 4.0f;
+					TSharedPtr<SNotificationItem> Notification = FSlateNotificationManager::Get().AddNotification(Info);
+					if (Notification)
+					{
+						Notification->SetCompletionState(SNotificationItem::CS_Fail);
+					}
+				}
+			}
+		}
 	}
 }
 
diff --git a/Engine/Source/Runtime/Engine/Private/Animation/AnimationSettings.cpp b/Engine/Source/Runtime/Engine/Private/Animation/AnimationSettings.cpp
--- a/Engine/Source/Runtime/Engine/Private/Animation/AnimationSettings.cpp
+++ b/Engine/Source/Runtime/Engine/Private/Animation/AnimationSettings.cpp
@@ -30,6 +30,10 @@ UAnimationSettings::UAnimationSettings(const FObjectInitializer& ObjectInitializ
 		FMirrorFindReplaceExpression("_R", "_L", EMirrorFindReplaceMethod::Suffix), FMirrorFindReplaceExpression("_L", "_R", EMirrorFindReplaceMethod::Suffix),
 		FMirrorFindReplaceExpression("right", "left", EMirrorFindReplaceMethod::Prefix), FMirrorFindReplaceExpression("left", "right",  EMirrorFindReplaceMethod::Prefix),
 		FMirrorFindReplaceExpression("Right", "Left", EMirrorFindReplaceMethod::Prefix), FMirrorFindReplaceExpression("Left", "Right",  EMirrorFindReplaceMethod::Prefix),
+		FMirrorFindReplaceExpression("([^}]*)_l_([^}]*)", "$1_r_$2", EMirrorFindReplaceMethod::RegularExpression), FMirrorFindReplaceExpression("([^}]*)_r_([^}]*)", "$1_l_$2", EMirrorFindReplaceMethod::RegularExpression),
+		FMirrorFindReplaceExpression("([^}]*)_L_([^}]*)", "$1_R_$2", EMirrorFindReplaceMethod::RegularExpression), FMirrorFindReplaceExpression("([^}]*)_R_([^}]*)", "$1_L_$2", EMirrorFindReplaceMethod::RegularExpression),
+		FMirrorFindReplaceExpression("([^}]*)_left_([^}]*)", "$1_right_$2", EMirrorFindReplaceMethod::RegularExpression), FMirrorFindReplaceExpression("([^}]*)_right_([^}]*)", "$1_left_$2", EMirrorFindReplaceMethod::RegularExpression),
+		FMirrorFindReplaceExpression("([^}]*)_Left_([^}]*)", "$1_Right_$2", EMirrorFindReplaceMethod::RegularExpression), FMirrorFindReplaceExpression("([^}]*)_Right_([^}]*)", "$1_Left_$2", EMirrorFindReplaceMethod::RegularExpression),
 		FMirrorFindReplaceExpression("((?:^[sS]pine|^[rR]oot|^[pP]elvis|^[nN]eck|^[hH]ead|^ik_hand_gun).*)", "$1", EMirrorFindReplaceMethod::RegularExpression)
 	};
 }
