[UE][FIX] Add range checks to PhAT Mirror function

#rmx

diff --git a/Engine/Source/Editor/PhysicsAssetEditor/Private/PhysicsAssetEditorSharedData.cpp b/Engine/Source/Editor/PhysicsAssetEditor/Private/PhysicsAssetEditorSharedData.cpp
--- a/Engine/Source/Editor/PhysicsAssetEditor/Private/PhysicsAssetEditorSharedData.cpp
+++ b/Engine/Source/Editor/PhysicsAssetEditor/Private/PhysicsAssetEditorSharedData.cpp
@@ -340,8 +340,6 @@ void FPhysicsAssetEditorSharedData::Mirror()
 	USkeletalMesh* EditorSkelMesh = PhysicsAsset->GetPreviewMesh();
 	if(EditorSkelMesh)
 	{
-		TArray<int32> DBG_SrcCollidingBodies;
-
 		// Build list of all bodies and constraints to be mirrored
 		TArray<FMirrorInfo> MirrorInfos;
 		MirrorInfos.Reserve(SelectedBodies.Num() + SelectedConstraints.Num());
@@ -445,9 +443,20 @@ void FPhysicsAssetEditorSharedData::Mirror()
 					for(FName SourceCollidingBoneName : MirrorInfo.CollidingBodyBoneNames)
 					{
 						const int32 SourceCollidingBoneIndex = EditorSkelMesh->GetRefSkeleton().FindBoneIndex(SourceCollidingBoneName); // Find Index of the bone associated with the body that the source body was allowed to collide with.
-						const int32 MirrorCollidingBoneIndex = PhysicsAsset->FindMirroredBone(EditorSkelMesh, SourceCollidingBoneIndex); // Find the index of the bone that mirrors the colliding body's bone.
-						const FName MirrorCollidingBoneName = EditorSkelMesh->GetRefSkeleton().GetBoneName(MirrorCollidingBoneIndex); // Find the name of the bone that mirrors the colliding body's bone.
-						const int32 MirrorCollidingBodyIndex = PhysicsAsset->FindBodyIndex(MirrorCollidingBoneName); // Find the index of the colliding body.
+
+						int32 MirrorCollidingBoneIndex = INDEX_NONE;
+						if (EditorSkelMesh->GetRefSkeleton().IsValidIndex(SourceCollidingBoneIndex))
+						{
+							MirrorCollidingBoneIndex = PhysicsAsset->FindMirroredBone(EditorSkelMesh, SourceCollidingBoneIndex); // Find the index of the bone that mirrors the colliding body's bone.
+						}
+
+						FName MirrorCollidingBoneName = NAME_None;
+						if (EditorSkelMesh->GetRefSkeleton().IsValidIndex(MirrorCollidingBoneIndex))
+						{
+							MirrorCollidingBoneName = EditorSkelMesh->GetRefSkeleton().GetBoneName(MirrorCollidingBoneIndex); // Find the name of the bone that mirrors the colliding body's bone.
+						}
+
+                        const int32 MirrorCollidingBodyIndex = PhysicsAsset->FindBodyIndex(MirrorCollidingBoneName); // Find the index of the colliding body.;
 
 						if (MirrorCollidingBodyIndex != INDEX_NONE)
 						{
