[UE][FIX] Update keybord shortcut text for PhAT constraint editing

#rnx

diff --git a/Engine/Source/Editor/PhysicsAssetEditor/Private/PhysicsAssetEditorEditMode.cpp b/Engine/Source/Editor/PhysicsAssetEditor/Private/PhysicsAssetEditorEditMode.cpp
--- a/Engine/Source/Editor/PhysicsAssetEditor/Private/PhysicsAssetEditorEditMode.cpp
+++ b/Engine/Source/Editor/PhysicsAssetEditor/Private/PhysicsAssetEditorEditMode.cpp
@@ -579,11 +579,11 @@ void FPhysicsAssetEditorEditMode::DrawHUD(FEditorViewportClient* ViewportClient,
 	{
 		if (ViewportClient->GetWidgetMode() == UE::Widget::WM_Translate)
 		{
-			TextItem.Text = LOCTEXT("SingleMove", "Hold ALT to move a single reference frame");
+			TextItem.Text = LOCTEXT("SingleMove", "Hold ALT to move parent reference frame, SHIFT + ALT to move child reference frame");
 		}
 		else if (ViewportClient->GetWidgetMode() == UE::Widget::WM_Rotate)
 		{
-			TextItem.Text = LOCTEXT("SingleRotate", "Hold ALT to rotate single reference frame");
+			TextItem.Text = LOCTEXT("SingleRotate", "Hold ALT to rotate parent reference frame, SHIFT + ALT to rotate child reference frame");
 		}
 	}
 
