Chaos - fix PhAT Enable/DisableCollisions not updating correctly

diff --git a/Engine/Source/Editor/PhysicsAssetEditor/Private/PhysicsAssetEditorSharedData.cpp b/Engine/Source/Editor/PhysicsAssetEditor/Private/PhysicsAssetEditorSharedData.cpp
--- a/Engine/Source/Editor/PhysicsAssetEditor/Private/PhysicsAssetEditorSharedData.cpp
+++ b/Engine/Source/Editor/PhysicsAssetEditor/Private/PhysicsAssetEditorSharedData.cpp
@@ -1249,6 +1249,8 @@ void FPhysicsAssetEditorSharedData::SetCollisionBetweenSelected(bool bEnableColl
 
 
 	UpdateNoCollisionBodies();
+	
+	RefreshPhysicsAssetChange(PhysicsAsset);
 
 	BroadcastPreviewChanged();
 }
@@ -1301,6 +1303,8 @@ void FPhysicsAssetEditorSharedData::SetCollisionBetweenSelectedAndAll(bool bEnab
 
 	UpdateNoCollisionBodies();
 
+	RefreshPhysicsAssetChange(PhysicsAsset);
+
 	BroadcastPreviewChanged();
 }
 
@@ -1346,6 +1350,8 @@ void FPhysicsAssetEditorSharedData::SetCollisionBetween(int32 Body1Index, int32
 		}
 
 		UpdateNoCollisionBodies();
+	
+		RefreshPhysicsAssetChange(PhysicsAsset);
 	}
 
 	BroadcastPreviewChanged();
