[UE][FIX] Alt rotate does not work when simulating in PhAT and the mouse cursor is first clicked when hovering over a physics body
- Change FPhysicsAssetEditorEditMode::SimMousePress s.t. it only consumes input (returns true) when it has an effect.

#rnx

diff --git a/Engine/Source/Editor/PhysicsAssetEditor/Private/PhysicsAssetEditorEditMode.cpp b/Engine/Source/Editor/PhysicsAssetEditor/Private/PhysicsAssetEditorEditMode.cpp
--- a/Engine/Source/Editor/PhysicsAssetEditor/Private/PhysicsAssetEditorEditMode.cpp
+++ b/Engine/Source/Editor/PhysicsAssetEditor/Private/PhysicsAssetEditorEditMode.cpp
@@ -848,6 +848,8 @@ void FPhysicsAssetEditorEditMode::OpenSelectionMenu(FEditorViewportClient* InVie
 
 bool FPhysicsAssetEditorEditMode::SimMousePress(FEditorViewportClient* InViewportClient, FKey Key)
 {
+	bool bHandled = false;
+
 	FViewport* Viewport = InViewportClient->Viewport;
 
 	bool bCtrlDown = Viewport->KeyState(EKeys::LeftControl) || Viewport->KeyState(EKeys::RightControl);
@@ -911,13 +913,12 @@ bool FPhysicsAssetEditorEditMode::SimMousePress(FEditorViewportClient* InViewpor
 			{
 				SharedData->EditorSkelComp->AddImpulseAtLocation(Click.GetDirection() * SharedData->EditorOptions->PokeStrength, Result.Location, BoneName);
 			}
+
+			bHandled = true;
 		}
 	}
 
-	// @todo(ccaulfield): really this should return false if we don't have Ctrl or Shift help down. This would allow the mouse-fly
-	// behaviour to work even when clicking on a space occupied by a body. However we don't want to change this until we have a way
-	// to enable/disable the fly and orbit camera behaviours.
-	return bHit;
+	return bHandled;
 }
 
 void FPhysicsAssetEditorEditMode::SimMouseMove(FEditorViewportClient* InViewportClient, float DeltaX, float DeltaY)
