Crash fix for attempting to find the parent index of index none.

diff --git a/Engine/Source/Editor/PhysicsAssetEditor/Private/PhysicsAssetEditorMode.cpp b/Engine/Source/Editor/PhysicsAssetEditor/Private/PhysicsAssetEditorMode.cpp
--- a/Engine/Source/Editor/PhysicsAssetEditor/Private/PhysicsAssetEditorMode.cpp
+++ b/Engine/Source/Editor/PhysicsAssetEditor/Private/PhysicsAssetEditorMode.cpp
@@ -122,7 +122,7 @@ FText BuildCrossConstraintCountText(const UPhysicsAsset* const PhysicsAsset)
 					const int32 ChildBoneIndex = RefSkeleton.FindBoneIndex(CurrentConstraintSetup->DefaultInstance.GetChildBoneName());
 					const int32 ParentBoneIndex = RefSkeleton.FindBoneIndex(CurrentConstraintSetup->DefaultInstance.GetParentBoneName());
 
-					if (ParentBoneIndex != RefSkeleton.GetParentIndex(ChildBoneIndex))
+					if (RefSkeleton.IsValidIndex(ChildBoneIndex) && (ParentBoneIndex != RefSkeleton.GetParentIndex(ChildBoneIndex)))
 					{
 						++CrossConstraintCount;
 					}
