[UE][FIX] Vewport manipulation widget does not appear in the same loaction as the selected Constraints on bones with scaling in PhAT
- Don't ignore bone scaling when calculating widget transforms for constraints.

#rnx

diff --git a/Engine/Source/Editor/Persona/Private/PhysicsAssetRenderUtils.cpp b/Engine/Source/Editor/Persona/Private/PhysicsAssetRenderUtils.cpp
--- a/Engine/Source/Editor/Persona/Private/PhysicsAssetRenderUtils.cpp
+++ b/Engine/Source/Editor/Persona/Private/PhysicsAssetRenderUtils.cpp
@@ -34,7 +34,6 @@ FTransform GetConstraintMatrix(const USkeletalMeshComponent* const SkeletalMeshC
 
 	FTransform LFrame = ConstraintSetup->DefaultInstance.GetRefFrame(Frame);
 	FTransform BoneTM = SkeletalMeshComponent->GetBoneTransform(BoneIndex);
-	BoneTM.RemoveScaling();
 
 	return LFrame * BoneTM;
 }
@@ -474,6 +473,10 @@ namespace PhysicsAssetRender
 							FTransform Con1Frame = GetConstraintMatrix(SkeletalMeshComponent, ConstraintSetup, EConstraintFrame::Frame1, BoneIndex1);
 							FTransform Con2Frame = GetConstraintMatrix(SkeletalMeshComponent, ConstraintSetup, EConstraintFrame::Frame2, BoneIndex2);
 
+							// Remove scaling from constraint frame transforms so that constraint limit cones etc are all drawn at the same scale.
+							Con1Frame.RemoveScaling();
+							Con2Frame.RemoveScaling();
+
 							ConstraintSetup->DefaultInstance.DrawConstraint(PDI, RenderSettings->ConstraintDrawSize, DrawScale, bDrawLimits, bDrawSelected, Con1Frame, Con2Frame, RenderSettings->bShowConstraintsAsPoints);
 
 						}
diff --git a/Engine/Source/Editor/PhysicsAssetEditor/Private/PhysicsAssetEditorSharedData.cpp b/Engine/Source/Editor/PhysicsAssetEditor/Private/PhysicsAssetEditorSharedData.cpp
--- a/Engine/Source/Editor/PhysicsAssetEditor/Private/PhysicsAssetEditorSharedData.cpp
+++ b/Engine/Source/Editor/PhysicsAssetEditor/Private/PhysicsAssetEditorSharedData.cpp
@@ -2293,8 +2293,6 @@ void FPhysicsAssetEditorSharedData::SetConstraintRelTM(const FPhysicsAssetEditor
 	if (BoneIndex != INDEX_NONE)
 	{
 		FTransform BoneTM = EditorSkelComp->GetBoneTransform(BoneIndex);
-		BoneTM.RemoveScaling();
-
 		ConstraintSetup->DefaultInstance.SetRefFrame(EConstraintFrame::Frame1, WNewChildFrame.GetRelativeTransform(BoneTM));
 	}
 }
@@ -2651,46 +2649,25 @@ FTransform FPhysicsAssetEditorSharedData::GetConstraintBodyTM(const UPhysicsCons
 	}
 }
 
-FTransform FPhysicsAssetEditorSharedData::GetConstraintWorldTM(const UPhysicsConstraintTemplate* ConstraintSetup, EConstraintFrame::Type Frame, float Scale) const
+FTransform FPhysicsAssetEditorSharedData::GetConstraintWorldTM(const UPhysicsConstraintTemplate* const ConstraintSetup, const EConstraintFrame::Type Frame, const float Scale) const
 {
-	if (ConstraintSetup == NULL)
-	{
-		return FTransform::Identity;
-	}
-
 	USkeletalMesh* EditorSkelMesh = PhysicsAsset->GetPreviewMesh();
-	if(EditorSkelMesh == nullptr)
-	{
-		return FTransform::Identity;
-	}
-
-	FVector Scale3D(Scale);
 
-	int32 BoneIndex;
-	FTransform LFrame = ConstraintSetup->DefaultInstance.GetRefFrame(Frame);
-	if (Frame == EConstraintFrame::Frame1)
+	if ((ConstraintSetup != nullptr) && (EditorSkelMesh != nullptr))
 	{
-		BoneIndex = EditorSkelMesh->GetRefSkeleton().FindBoneIndex(ConstraintSetup->DefaultInstance.ConstraintBone1);
-	}
-	else
-	{
-		BoneIndex = EditorSkelMesh->GetRefSkeleton().FindBoneIndex(ConstraintSetup->DefaultInstance.ConstraintBone2);
-	}
+		const FName BoneName = (Frame == EConstraintFrame::Frame1) ? ConstraintSetup->DefaultInstance.ConstraintBone1 : ConstraintSetup->DefaultInstance.ConstraintBone2;
+		const int32 BoneIndex = EditorSkelMesh->GetRefSkeleton().FindBoneIndex(BoneName);
 
-	// If we couldn't find the bone - fall back to identity.
-	if (BoneIndex == INDEX_NONE)
-	{
-		return FTransform::Identity;
+		if (BoneIndex != INDEX_NONE)
+		{	
+			FTransform LFrame = ConstraintSetup->DefaultInstance.GetRefFrame(Frame);
+			LFrame.ScaleTranslation(FVector(Scale));
+			const FTransform BoneTM = EditorSkelComp->GetBoneTransform(BoneIndex);
+			return LFrame * BoneTM;
+		}
 	}
-	else
-	{
-		FTransform BoneTM = EditorSkelComp->GetBoneTransform(BoneIndex);
-		BoneTM.RemoveScaling();
-
-		LFrame.ScaleTranslation(Scale3D);
 
-		return LFrame * BoneTM;
-	}
+	return FTransform::Identity;
 }
 
 FTransform FPhysicsAssetEditorSharedData::GetConstraintMatrix(int32 ConstraintIndex, EConstraintFrame::Type Frame, float Scale) const
diff --git a/Engine/Source/Editor/PhysicsAssetEditor/Private/PhysicsAssetEditorSharedData.h b/Engine/Source/Editor/PhysicsAssetEditor/Private/PhysicsAssetEditorSharedData.h
--- a/Engine/Source/Editor/PhysicsAssetEditor/Private/PhysicsAssetEditorSharedData.h
+++ b/Engine/Source/Editor/PhysicsAssetEditor/Private/PhysicsAssetEditorSharedData.h
@@ -96,7 +96,7 @@ public:
 	FTransform GetConstraintWorldTM(const FSelection* Constraint, EConstraintFrame::Type Frame) const;
 
 	/** Get the world transform of the specified constraint */
-	FTransform GetConstraintWorldTM(const UPhysicsConstraintTemplate* ConstraintSetup, EConstraintFrame::Type Frame, float Scale = 1.f) const;
+	FTransform GetConstraintWorldTM(const UPhysicsConstraintTemplate* const ConstraintSetup, const EConstraintFrame::Type Frame, const float Scale = 1.0f) const;
 
 	/** Get the world transform of the specified constraint */
 	FTransform GetConstraintMatrix(int32 ConstraintIndex, EConstraintFrame::Type Frame, float Scale) const;
