Fix bShowCoM draw mode in physics asset editor. Changed CoM star visualisation for foreground draw priority to show on top of the mesh in the editor.

diff --git a/Engine/Source/Editor/PhysicsAssetEditor/Private/PhysicsAssetEditor.cpp b/Engine/Source/Editor/PhysicsAssetEditor/Private/PhysicsAssetEditor.cpp
--- a/Engine/Source/Editor/PhysicsAssetEditor/Private/PhysicsAssetEditor.cpp
+++ b/Engine/Source/Editor/PhysicsAssetEditor/Private/PhysicsAssetEditor.cpp
@@ -2621,14 +2621,14 @@ bool FPhysicsAssetEditor::IsConstraintRenderingMode(EPhysicsAssetEditorConstrain
 
 void FPhysicsAssetEditor::OnToggleMassProperties()
 {
-	SharedData->bShowCOM = !SharedData->bShowCOM;
+	SharedData->ToggleShowCom();
 
 	RefreshPreviewViewport();
 }
 
 bool FPhysicsAssetEditor::IsToggleMassProperties() const
 {
-	return SharedData->bShowCOM;
+	return SharedData->GetShowCom();
 }
 
 void FPhysicsAssetEditor::OnSetCollision(bool bEnable)
diff --git a/Engine/Source/Editor/PhysicsAssetEditor/Private/PhysicsAssetEditorEditMode.cpp b/Engine/Source/Editor/PhysicsAssetEditor/Private/PhysicsAssetEditorEditMode.cpp
--- a/Engine/Source/Editor/PhysicsAssetEditor/Private/PhysicsAssetEditorEditMode.cpp
+++ b/Engine/Source/Editor/PhysicsAssetEditor/Private/PhysicsAssetEditorEditMode.cpp
@@ -603,7 +603,7 @@ void FPhysicsAssetEditorEditMode::DrawHUD(FEditorViewportClient* ViewportClient,
 	int32 HalfY = Viewport->GetSizeXY().Y / 2;
 
 	// If showing center-of-mass, and physics is started up..
-	if (SharedData->bShowCOM)
+	if (SharedData->GetShowCom())
 	{
 		// iterate over each bone
 		for (int32 i = 0; i <SharedData->EditorSkelComp->Bodies.Num(); ++i)
diff --git a/Engine/Source/Editor/PhysicsAssetEditor/Private/PhysicsAssetEditorSharedData.cpp b/Engine/Source/Editor/PhysicsAssetEditor/Private/PhysicsAssetEditorSharedData.cpp
--- a/Engine/Source/Editor/PhysicsAssetEditor/Private/PhysicsAssetEditorSharedData.cpp
+++ b/Engine/Source/Editor/PhysicsAssetEditor/Private/PhysicsAssetEditorSharedData.cpp
@@ -71,9 +71,6 @@ FPhysicsAssetEditorSharedData::FPhysicsAssetEditorSharedData()
 	, bSuspendSelectionBroadcast(false)
 	, InsideSelChange(0)
 {
-	// Editor variables
-	bShowCOM = false;
-
 	bRunningSimulation = false;
 	bNoGravitySimulation = false;
 
@@ -290,6 +287,32 @@ bool FPhysicsAssetEditorSharedData::ClipboardHasCompatibleData()
 	return ParseClipboard(DummyAsset, DummyObjectType, DummyObject);
 }
 
+void FPhysicsAssetEditorSharedData::ToggleShowCom()
+{
+	if(FPhysicsAssetRenderSettings* const PhysicsAssetRenderSettings = GetRenderSettings())
+	{
+		PhysicsAssetRenderSettings->bShowCOM = !PhysicsAssetRenderSettings->bShowCOM;
+	}
+}
+
+void FPhysicsAssetEditorSharedData::SetShowCom(bool InValue)
+{
+	if(FPhysicsAssetRenderSettings* const PhysicsAssetRenderSettings = GetRenderSettings())
+	{
+		PhysicsAssetRenderSettings->bShowCOM = InValue;
+	}
+}
+
+bool FPhysicsAssetEditorSharedData::GetShowCom() const
+{
+	if(FPhysicsAssetRenderSettings* const PhysicsAssetRenderSettings = GetRenderSettings())
+	{
+		return PhysicsAssetRenderSettings->bShowCOM;
+	}
+
+	return false;
+}
+
 bool FPhysicsAssetEditorSharedData::ParseClipboard(UPhysicsAsset*& OutAsset, FString& OutObjectType, UObject*& OutObject)
 {
 	FString ClipboardContent;
diff --git a/Engine/Source/Editor/PhysicsAssetEditor/Private/PhysicsAssetEditorSharedData.h b/Engine/Source/Editor/PhysicsAssetEditor/Private/PhysicsAssetEditorSharedData.h
--- a/Engine/Source/Editor/PhysicsAssetEditor/Private/PhysicsAssetEditorSharedData.h
+++ b/Engine/Source/Editor/PhysicsAssetEditor/Private/PhysicsAssetEditorSharedData.h
@@ -249,6 +249,11 @@ public:
 	/** Returns true if the clipboard contains data this class can process */
 	static bool ClipboardHasCompatibleData();
 
+	/** Control whether we draw a CoM marker in the viewport */
+	void ToggleShowCom();
+	void SetShowCom(bool InValue);
+	bool GetShowCom() const;
+
 private:
 	/** Initializes a constraint setup */
 	void InitConstraintSetup(UPhysicsConstraintTemplate* ConstraintSetup, int32 ChildBodyIndex, int32 ParentBodyIndex);
@@ -350,9 +355,6 @@ public:
 
 	struct FPhysicsAssetRenderSettings* GetRenderSettings() const;
 
-	/** Show flags */
-	bool bShowCOM;
-
 	/** Misc toggles */
 	bool bRunningSimulation;
 	bool bNoGravitySimulation;
diff --git a/Engine/Source/Runtime/Engine/Private/PhysicsEngine/BodyInstance.cpp b/Engine/Source/Runtime/Engine/Private/PhysicsEngine/BodyInstance.cpp
--- a/Engine/Source/Runtime/Engine/Private/PhysicsEngine/BodyInstance.cpp
+++ b/Engine/Source/Runtime/Engine/Private/PhysicsEngine/BodyInstance.cpp
@@ -2624,7 +2624,7 @@ void FBodyInstance::DrawCOMPosition(FPrimitiveDrawInterface* PDI, float COMRende
 {
 	if (IsValidBodyInstance())
 	{
-		DrawWireStar(PDI, GetCOMPosition(), COMRenderSize, COMRenderColor, SDPG_World);
+		DrawWireStar(PDI, GetCOMPosition(), COMRenderSize, COMRenderColor, SDPG_Foreground);
 	}
 }
 
