[UE] Fix for static analysis warning in FPhysicsAssetEditorSharedData::GetConstraintBodyTM - Horde issue 288292

#rnx

diff --git a/Engine/Source/Editor/PhysicsAssetEditor/Private/PhysicsAssetEditorSharedData.cpp b/Engine/Source/Editor/PhysicsAssetEditor/Private/PhysicsAssetEditorSharedData.cpp
--- a/Engine/Source/Editor/PhysicsAssetEditor/Private/PhysicsAssetEditorSharedData.cpp
+++ b/Engine/Source/Editor/PhysicsAssetEditor/Private/PhysicsAssetEditorSharedData.cpp
@@ -2638,7 +2638,7 @@ void FPhysicsAssetEditorSharedData::DeleteCurrentPrim()
 
 FTransform FPhysicsAssetEditorSharedData::GetConstraintBodyTM(const UPhysicsConstraintTemplate* ConstraintSetup, EConstraintFrame::Type Frame) const
 {
-	if ((ConstraintSetup != nullptr) && (EditorSkelComp == nullptr))
+	if ((ConstraintSetup != nullptr) && (EditorSkelComp != nullptr))
 	{
 		const FName BoneName = (Frame == EConstraintFrame::Frame1) ? ConstraintSetup->DefaultInstance.ConstraintBone1 : ConstraintSetup->DefaultInstance.ConstraintBone2;
 		const int32 BoneIndex = EditorSkelComp->GetBoneIndex(BoneName);
