Fix F framing of level sets in physics asset editor

diff --git a/Engine/Source/Editor/PhysicsAssetEditor/Private/PhysicsAssetEditorEditMode.cpp b/Engine/Source/Editor/PhysicsAssetEditor/Private/PhysicsAssetEditorEditMode.cpp
--- a/Engine/Source/Editor/PhysicsAssetEditor/Private/PhysicsAssetEditorEditMode.cpp
+++ b/Engine/Source/Editor/PhysicsAssetEditor/Private/PhysicsAssetEditorEditMode.cpp
@@ -82,6 +82,14 @@ bool FPhysicsAssetEditorEditMode::GetCameraTarget(FSphere& OutTarget) const
 		{
 			Bounds += AggGeom.TaperedCapsuleElems[SelectedObject.PrimitiveIndex].CalcAABB(BoneTM, Scale);
 		}
+		else if (SelectedObject.PrimitiveType == EAggCollisionShape::LevelSet)
+		{
+			Bounds += AggGeom.LevelSetElems[SelectedObject.PrimitiveIndex].CalcAABB(BoneTM, BoneTM.GetScale3D());
+		}
+		else if (SelectedObject.PrimitiveType == EAggCollisionShape::SkinnedLevelSet)
+		{
+			Bounds += AggGeom.SkinnedLevelSetElems[SelectedObject.PrimitiveIndex].CalcAABB(BoneTM, BoneTM.GetScale3D());
+		}
 
 		bHandled = true;
 	}
