[UE] Fix for PhAT creating physics bodies on the wrong bone
- Incorrect Bone Index was being passed to the MakeNewBody function.
- Changed FPhysicsAssetEditor::ResetBoneCollision to lookup bone index in the physics assets Ref Skeleton, instead of the EditorSkelComp, making it consistent with the MakeNewBody function.

#rnx

diff --git a/Engine/Source/Editor/PhysicsAssetEditor/Private/PhysicsAssetEditor.cpp b/Engine/Source/Editor/PhysicsAssetEditor/Private/PhysicsAssetEditor.cpp
--- a/Engine/Source/Editor/PhysicsAssetEditor/Private/PhysicsAssetEditor.cpp
+++ b/Engine/Source/Editor/PhysicsAssetEditor/Private/PhysicsAssetEditor.cpp
@@ -2333,7 +2333,7 @@ void FPhysicsAssetEditor::ResetBoneCollision()
 
 				UBoneProxy* BoneProxy = CastChecked<UBoneProxy>(BoneItem->GetObject());
 
-				int32 BoneIndex = SharedData->EditorSkelComp->GetBoneIndex(BoneProxy->BoneName);
+				const int32 BoneIndex = EditorSkelMesh->GetRefSkeleton().FindBoneIndex(BoneProxy->BoneName);
 				if (BoneIndex != INDEX_NONE)
 				{
 					SharedData->MakeNewBody(BoneIndex);
