[UE] CoM visualization and editing tools for PhAT
- Draw CoM positions and mass values in the viewport
- Add menu options to control drawing CoM indicators for all, selected or no physics bodies
- New feature to allow the CoM location to be fixed in world space in the editor.This allows users to move a physics body without changing the CoMs position in world space.
- Add support for moving a bodies CoM using a standard translation manipulator widget.

diff --git a/Engine/Source/Editor/DetailCustomizations/Private/BodyInstanceCustomization.cpp b/Engine/Source/Editor/DetailCustomizations/Private/BodyInstanceCustomization.cpp
--- a/Engine/Source/Editor/DetailCustomizations/Private/BodyInstanceCustomization.cpp
+++ b/Engine/Source/Editor/DetailCustomizations/Private/BodyInstanceCustomization.cpp
@@ -1318,7 +1318,13 @@ void FBodyInstanceCustomizationHelper::UpdateFilters()
 	}
 }
 
-void FBodyInstanceCustomizationHelper::CustomizeDetails( IDetailLayoutBuilder& DetailBuilder, TSharedRef<IPropertyHandle> BodyInstanceHandler)
+
+void FBodyInstanceCustomizationHelper::CustomizeDetails(IDetailLayoutBuilder& DetailBuilder, TSharedRef<IPropertyHandle> BodyInstanceHandler)
+{
+	CustomizeDetails(DetailBuilder, BodyInstanceHandler, TFunction<void(TSharedRef<IPropertyHandle>)>());
+}
+
+void FBodyInstanceCustomizationHelper::CustomizeDetails(IDetailLayoutBuilder& DetailBuilder, TSharedRef<IPropertyHandle> BodyInstanceHandler, TFunction<void(TSharedRef<IPropertyHandle>)> CustomizeCoMNudge)
 {
 	if( BodyInstanceHandler->IsValidHandle() )
 	{
@@ -1365,7 +1371,16 @@ void FBodyInstanceCustomizationHelper::CustomizeDetails( IDetailLayoutBuilder& D
 			.Visibility(TAttribute<EVisibility>(this, &FBodyInstanceCustomizationHelper::IsAutoWeldVisible));
 
 		EnablePhysicsProperty(GET_MEMBER_NAME_CHECKED(FBodyInstance, bStartAwake), bDisplayStartAwake);
-		EnablePhysicsProperty(GET_MEMBER_NAME_CHECKED(FBodyInstance, COMNudge), bDisplayCOMNudge);
+
+		if (CustomizeCoMNudge && bDisplayCOMNudge)
+		{
+			CustomizeCoMNudge(BodyInstanceHandler);
+		}
+		else
+		{
+			EnablePhysicsProperty(GET_MEMBER_NAME_CHECKED(FBodyInstance, COMNudge), bDisplayCOMNudge);
+		}
+
 		EnablePhysicsProperty(GET_MEMBER_NAME_CHECKED(FBodyInstance, MassScale), bDisplayMassScale);
 
 		if (bDisplayMaxAngularVelocity)
diff --git a/Engine/Source/Editor/DetailCustomizations/Private/BodyInstanceCustomization.h b/Engine/Source/Editor/DetailCustomizations/Private/BodyInstanceCustomization.h
--- a/Engine/Source/Editor/DetailCustomizations/Private/BodyInstanceCustomization.h
+++ b/Engine/Source/Editor/DetailCustomizations/Private/BodyInstanceCustomization.h
@@ -161,7 +161,8 @@ class FBodyInstanceCustomizationHelper  : public TSharedFromThis<FBodyInstanceCu
 {
 public:
 	FBodyInstanceCustomizationHelper(const TArray<TWeakObjectPtr<UObject>>& InObjectsCustomized);
-	void CustomizeDetails( IDetailLayoutBuilder& DetailBuilder, TSharedRef<IPropertyHandle> BodyInstanceHandler );
+	void CustomizeDetails( IDetailLayoutBuilder& DetailBuilder, TSharedRef<IPropertyHandle> BodyInstanceHandler);
+	void CustomizeDetails(IDetailLayoutBuilder& DetailBuilder, TSharedRef<IPropertyHandle> BodyInstanceHandler, TFunction<void(TSharedRef<IPropertyHandle>)> CustomizeCoMNudge);
 
 private:
 	void UpdateFilters();
diff --git a/Engine/Source/Editor/DetailCustomizations/Private/BodySetupDetails.cpp b/Engine/Source/Editor/DetailCustomizations/Private/BodySetupDetails.cpp
--- a/Engine/Source/Editor/DetailCustomizations/Private/BodySetupDetails.cpp
+++ b/Engine/Source/Editor/DetailCustomizations/Private/BodySetupDetails.cpp
@@ -54,7 +54,7 @@ void FBodySetupDetails::CustomizeDetails( IDetailLayoutBuilder& DetailBuilder )
 			TSharedPtr<IPropertyHandle> BodyInstanceHandler = DetailBuilder.GetProperty(GET_MEMBER_NAME_CHECKED(UBodySetup, DefaultInstance));
 			
 			BodyInstanceCustomizationHelper = MakeShareable(new FBodyInstanceCustomizationHelper(ObjectsCustomized));
-			BodyInstanceCustomizationHelper->CustomizeDetails(DetailBuilder, DetailBuilder.GetProperty(GET_MEMBER_NAME_CHECKED(UBodySetup, DefaultInstance)));
+			BodyInstanceCustomizationHelper->CustomizeDetails(DetailBuilder, BodyInstanceHandler.ToSharedRef(), [this, &DetailBuilder](TSharedRef<IPropertyHandle> BodyInstanceHandle) { this->CustomizeCoMNudge(DetailBuilder, BodyInstanceHandle); });
 
 			IDetailCategoryBuilder& CollisionCategory = DetailBuilder.EditCategory("Collision");
 			DetailBuilder.HideProperty(BodyInstanceHandler);
@@ -113,6 +113,7 @@ void FBodySetupDetails::CustomizeDetails( IDetailLayoutBuilder& DetailBuilder )
 	}
 }
 
+
 TSharedRef<IDetailCustomization> FSkeletalBodySetupDetails::MakeInstance()
 {
 	return MakeShareable(new FSkeletalBodySetupDetails);
diff --git a/Engine/Source/Editor/DetailCustomizations/Private/BodySetupDetails.h b/Engine/Source/Editor/DetailCustomizations/Private/BodySetupDetails.h
--- a/Engine/Source/Editor/DetailCustomizations/Private/BodySetupDetails.h
+++ b/Engine/Source/Editor/DetailCustomizations/Private/BodySetupDetails.h
@@ -12,7 +12,7 @@ class IDetailLayoutBuilder;
 class IPropertyHandle;
 class UObject;
 
-class FBodySetupDetails : public IDetailCustomization
+class DETAILCUSTOMIZATIONS_API FBodySetupDetails : public IDetailCustomization
 {
 public:
 	/** Makes a new instance of this detail layout class for a specific detail view requesting it */
@@ -21,6 +21,8 @@ public:
 	/** IDetailCustomization interface */
 	virtual void CustomizeDetails( IDetailLayoutBuilder& DetailBuilder ) override;
 
+	virtual void CustomizeCoMNudge(IDetailLayoutBuilder& DetailBuilder, TSharedRef<IPropertyHandle> BodyInstanceHandler) {};
+
 private:
 	void AddPhysicalAnimation(IDetailLayoutBuilder& DetailBuilder);
 
diff --git a/Engine/Source/Editor/Persona/Private/PhysicsAssetRenderUtils.cpp b/Engine/Source/Editor/Persona/Private/PhysicsAssetRenderUtils.cpp
--- a/Engine/Source/Editor/Persona/Private/PhysicsAssetRenderUtils.cpp
+++ b/Engine/Source/Editor/Persona/Private/PhysicsAssetRenderUtils.cpp
@@ -41,11 +41,19 @@ FTransform GetConstraintMatrix(const USkeletalMeshComponent* const SkeletalMeshC
 	return LFrame * BoneTM;
 }
 
+void DrawWireStar(FPrimitiveDrawInterface* PDI, const FVector& Position, float Size, const FLinearColor& Color, uint8 DepthPriority, const float Thickness)
+{
+	PDI->DrawLine(Position + Size * FVector(1, 0, 0), Position - Size * FVector(1, 0, 0), Color, DepthPriority, Thickness);
+	PDI->DrawLine(Position + Size * FVector(0, 1, 0), Position - Size * FVector(0, 1, 0), Color, DepthPriority, Thickness);
+	PDI->DrawLine(Position + Size * FVector(0, 0, 1), Position - Size * FVector(0, 0, 1), Color, DepthPriority, Thickness);
+}
+
 ////////////////////////////////////////
 // struct FPhysicsAssetRenderSettings
 
 FPhysicsAssetRenderSettings::FPhysicsAssetRenderSettings()
-	: CollisionViewMode(EPhysicsAssetEditorCollisionViewMode::Solid)
+	: CenterOfMassViewMode(EPhysicsAssetEditorCenterOfMassViewMode::None)
+	, CollisionViewMode(EPhysicsAssetEditorCollisionViewMode::Solid)
 	, ConstraintViewMode(EPhysicsAssetEditorConstraintViewMode::AllLimits)
 	, ConstraintViewportManipulationFlags(EConstraintTransformComponentFlags::All)
 	, ConstraintTransformComponentDisplayRelativeToDefaultFlags(EConstraintTransformComponentFlags::None)
@@ -53,14 +61,16 @@ FPhysicsAssetRenderSettings::FPhysicsAssetRenderSettings()
 	, PhysicsBlend(1.0f)
 	, bHideKinematicBodies(false)
 	, bHideSimulatedBodies(false)
+	, bHideBodyMass(false)
 	, bRenderOnlySelectedConstraints(false)
-	, bShowCOM(false)
 	, bShowConstraintsAsPoints(false)
 	, bDrawViolatedLimits(false)
 	, BoneUnselectedColor(170, 155, 225)
 	, NoCollisionColor(200, 200, 200)
 	, COMRenderColor(255, 255, 100)
-	, COMRenderSize(5.0f)
+	, COMRenderSize(3.0f)
+	, COMRenderLineThickness(0.2f)
+	, COMRenderMassTextOffsetScreenspace(8.0f)
 	, InfluenceLineLength(2.0f)
 	, BoneUnselectedMaterial(nullptr)
 	, BoneNoCollisionMaterial(nullptr)
@@ -254,12 +264,14 @@ void FPhysicsAssetRenderSettings::ResetEditorViewportOptions()
 {
 	const FPhysicsAssetRenderSettings DefaultObject = FPhysicsAssetRenderSettings();
 
+	CenterOfMassViewMode = DefaultObject.CenterOfMassViewMode;
 	CollisionViewMode = DefaultObject.CollisionViewMode;
 	ConstraintViewMode = DefaultObject.ConstraintViewMode;
 	ConstraintDrawSize = DefaultObject.ConstraintDrawSize;
 	PhysicsBlend = DefaultObject.PhysicsBlend;
 	bHideKinematicBodies = DefaultObject.bHideKinematicBodies;
 	bHideSimulatedBodies = DefaultObject.bHideSimulatedBodies;
+	bHideBodyMass = DefaultObject.bHideBodyMass;
 	bRenderOnlySelectedConstraints = DefaultObject.bRenderOnlySelectedConstraints;
 	bShowConstraintsAsPoints = DefaultObject.bShowConstraintsAsPoints;
 	bDrawViolatedLimits = DefaultObject.bDrawViolatedLimits;
@@ -429,10 +441,43 @@ namespace PhysicsAssetRender
 	
 				PDI->SetHitProxy(NULL);
 			}
+		}
+	}
+
+	void DebugDrawCenterOfMass(USkeletalMeshComponent* const SkeletalMeshComponent, class UPhysicsAsset* const PhysicsAsset, FPrimitiveDrawInterface* PDI, TFunctionRef< FVector(const uint32) > GetCoMPosition, TFunctionRef< bool(const uint32) > IsSelected, CreateCoMHitProxyFn CreateHitProxy)
+	{
+		static const float SelectedItemRenderSizeMultiplier = 1.5f;
+
+		const FPhysicsAssetRenderSettings* const RenderSettings = UPhysicsAssetRenderUtilities::GetSettings(PhysicsAsset);
+
+		if (RenderSettings && (RenderSettings->CenterOfMassViewMode != EPhysicsAssetEditorCenterOfMassViewMode::None))
+		{
+			const bool bDrawSelectedOnly = RenderSettings->CenterOfMassViewMode == EPhysicsAssetEditorCenterOfMassViewMode::Selected;
 
-			if (RenderSettings->bShowCOM && SkeletalMeshComponent->Bodies.IsValidIndex(i))
+			for (int32 BodyIndex = 0, BodyCount = SkeletalMeshComponent->Bodies.Num(); BodyIndex < BodyCount; ++BodyIndex)
 			{
-				SkeletalMeshComponent->Bodies[i]->DrawCOMPosition(PDI, RenderSettings->COMRenderSize, RenderSettings->COMRenderColor);
+				const bool bIsSelected = IsSelected(BodyIndex);
+
+				if (!bDrawSelectedOnly || bIsSelected)
+				{
+					const FBodyInstance* const BodyInstance = SkeletalMeshComponent->Bodies[BodyIndex];
+
+					if (BodyInstance && BodyInstance->IsValidBodyInstance())
+					{
+						float COMRenderSize = RenderSettings->COMRenderSize;
+						float COMRenderLineThickness = RenderSettings->COMRenderLineThickness;
+
+						if (bIsSelected)
+						{
+							COMRenderSize *= SelectedItemRenderSizeMultiplier;
+							COMRenderLineThickness *= SelectedItemRenderSizeMultiplier;
+						}
+
+						PDI->SetHitProxy(CreateHitProxy(BodyIndex));
+						DrawWireStar(PDI, GetCoMPosition(BodyIndex), COMRenderSize, RenderSettings->COMRenderColor, SDPG_Foreground, COMRenderLineThickness);
+						PDI->SetHitProxy(NULL);
+					}
+				}
 			}
 		}
 	}
diff --git a/Engine/Source/Editor/Persona/Public/PhysicsAssetRenderUtils.h b/Engine/Source/Editor/Persona/Public/PhysicsAssetRenderUtils.h
--- a/Engine/Source/Editor/Persona/Public/PhysicsAssetRenderUtils.h
+++ b/Engine/Source/Editor/Persona/Public/PhysicsAssetRenderUtils.h
@@ -11,6 +11,7 @@
 #include "SceneManagement.h"
 #include "PhysicsAssetRenderUtils.generated.h"
 
+enum class EPhysicsAssetEditorCenterOfMassViewMode : uint8;
 enum class EPhysicsAssetEditorConstraintViewMode : uint8;
 enum class EPhysicsAssetEditorCollisionViewMode : uint8;
 
@@ -74,6 +75,9 @@ public:
 	void ResetEditorViewportOptions();
 
 	// Physics Asset Editor Viewport Options
+	UPROPERTY()
+	EPhysicsAssetEditorCenterOfMassViewMode CenterOfMassViewMode;
+
 	UPROPERTY()
 	EPhysicsAssetEditorCollisionViewMode CollisionViewMode;
 
@@ -99,12 +103,15 @@ public:
 
 	UPROPERTY()
 	bool bHideSimulatedBodies;
+	
+	UPROPERTY()
+	bool bHideBodyMass;
 
 	UPROPERTY()
 	bool bRenderOnlySelectedConstraints;
 
 	UPROPERTY()
-	bool bShowCOM;
+	bool bShowCOM_DEPRECATED;
 
 	UPROPERTY()
 	bool bShowConstraintsAsPoints;
@@ -125,6 +132,12 @@ public:
 	UPROPERTY()
 	float COMRenderSize;
 
+	UPROPERTY()
+	float COMRenderLineThickness;
+
+	UPROPERTY()
+	float COMRenderMassTextOffsetScreenspace;
+
 	UPROPERTY()
 	float InfluenceLineLength;
 
@@ -199,13 +212,18 @@ namespace PhysicsAssetRender
 	using GetPrimitiveTransformRef = TFunctionRef< FTransform (const UPhysicsAsset* PhysicsAsset, const FTransform& BoneTM, const int32 BodyIndex, const EAggCollisionShape::Type PrimType, const int32 PrimIndex, const float Scale) >;
 	using CreateBodyHitProxyFn = TFunctionRef< HHitProxy* (const int32 BodyIndex, const EAggCollisionShape::Type PrimitiveType, const int32 PrimitiveIndex) >;
 	using CreateConstraintHitProxyFn = TFunctionRef< HHitProxy* (const int32 InConstraintIndex) >;
+	using CreateCoMHitProxyFn = TFunctionRef< HHitProxy* (const int32 InBodyIndex) >;
+	using COMAccessorFunctionFn = TFunctionRef< TPair< bool, FVector >(const int32) >;
 
 	/** Debug draw Physics Asset bodies and constraints using the default callbacks */
 	PERSONA_API void DebugDraw(class USkeletalMeshComponent* const SkeletalMeshComponent, class UPhysicsAsset* const PhysicsAsset, FPrimitiveDrawInterface* PDI);
 	
 	/** Debug draw Physics Asset bodies using the supplied custom callbacks */
 	PERSONA_API void DebugDrawBodies(class USkeletalMeshComponent* const SkeletalMeshComponent, class UPhysicsAsset* const PhysicsAsset, FPrimitiveDrawInterface* PDI, GetPrimitiveRef< FColor > GetPrimitiveColor, GetPrimitiveRef< class UMaterialInterface* > GetPrimitiveMaterial, GetPrimitiveTransformRef GetPrimitiveTransform, CreateBodyHitProxyFn CreateHitProxy);
-	
+
+	/** Debug draw the Center of Masses for all bodies in the supplied Skeletal Mesh */
+	PERSONA_API void DebugDrawCenterOfMass(USkeletalMeshComponent* const SkeletalMeshComponent, class UPhysicsAsset* const PhysicsAsset, FPrimitiveDrawInterface* PDI, TFunctionRef< FVector(const uint32) > GetCoMPosition, TFunctionRef< bool(const uint32) > IsSelected, CreateCoMHitProxyFn CreateHitProxy);
+
 	/** Debug draw Physics Asset constraints using the supplied custom callbacks */
 	PERSONA_API void DebugDrawConstraints(class USkeletalMeshComponent* const SkeletalMeshComponent, class UPhysicsAsset* const PhysicsAsset, FPrimitiveDrawInterface* PDI, TFunctionRef< bool(const uint32) > IsConstraintSelected, const bool bRunningSimulation, CreateConstraintHitProxyFn CreateHitProxy);
 	
diff --git a/Engine/Source/Editor/PhysicsAssetEditor/Private/PhysicsAssetEditor.cpp b/Engine/Source/Editor/PhysicsAssetEditor/Private/PhysicsAssetEditor.cpp
--- a/Engine/Source/Editor/PhysicsAssetEditor/Private/PhysicsAssetEditor.cpp
+++ b/Engine/Source/Editor/PhysicsAssetEditor/Private/PhysicsAssetEditor.cpp
@@ -269,7 +269,7 @@ TSharedPtr<FPhysicsAssetEditorSharedData> FPhysicsAssetEditor::GetSharedData() c
 	return SharedData;
 }
 
-void FPhysicsAssetEditor::HandleViewportSelectionChanged(const TArray<FPhysicsAssetEditorSharedData::FSelection>& InSelectedBodies, const TArray<FPhysicsAssetEditorSharedData::FSelection>& InSelectedConstraints)
+void FPhysicsAssetEditor::HandleViewportSelectionChanged(const TArray<FPhysicsAssetEditorSharedData::FSelection>& InSelectedBodies, const TArray<FPhysicsAssetEditorSharedData::FSelection>& InSelectedConstraints, const TArray<FPhysicsAssetEditorSharedData::FSelection>& InSelectedCoMs)
 {
 	if (!bSelecting)
 	{
@@ -280,7 +280,7 @@ void FPhysicsAssetEditor::HandleViewportSelectionChanged(const TArray<FPhysicsAs
 			SkeletonTree->DeselectAll();
 		}
 
-		if(InSelectedBodies.Num() == 0 && InSelectedConstraints.Num() == 0)
+		if(InSelectedBodies.IsEmpty() && InSelectedConstraints.IsEmpty() && InSelectedCoMs.IsEmpty())
 		{
 			if (PhysAssetProperties.IsValid())
 			{
@@ -824,6 +824,31 @@ void FPhysicsAssetEditor::ExtendViewportMenus()
 				}
 			}));
 
+			Section.AddSubMenu(TEXT("CenterOfMassModeSubMenu"), LOCTEXT("CenterOfMassModeSubMenu", "Center of Mass"), FText::GetEmpty(),
+				FNewToolMenuDelegate::CreateLambda([WeakPhysicsAssetEditor = PhysicsAssetEditor.ToWeakPtr()](UToolMenu* InSubMenu)
+			{
+				const FPhysicsAssetEditorCommands& Commands = FPhysicsAssetEditorCommands::Get();
+
+				{
+					FToolMenuSection& Section = InSubMenu->AddSection("PhysicsAssetEditorCenterOfMassRenderSettings", LOCTEXT("CenterOfMassRenderSettingsHeader", "Center of Mass Drawing"));
+					Section.AddMenuEntry(Commands.HideBodyMass);
+				}
+
+				{
+					FToolMenuSection& Section = InSubMenu->AddSection("PhysicsAssetEditorCenterOfMassRenderingMode", LOCTEXT("CenterOfMassRenderingModeHeader", "Center of Mass Drawing (Edit)"));
+					Section.AddMenuEntry(Commands.CenterOfMassRenderingMode_All);
+					Section.AddMenuEntry(Commands.CenterOfMassRenderingMode_Selected);
+					Section.AddMenuEntry(Commands.CenterOfMassRenderingMode_None);
+				}
+
+				{
+					FToolMenuSection& Section = InSubMenu->AddSection("PhysicsAssetEditorCenterOfMassRenderingModeSim", LOCTEXT("CenterOfMassRenderingModeSimHeader", "Center of Mass Drawing (Simulation)"));
+					Section.AddMenuEntry(Commands.CenterOfMassRenderingMode_Simulation_All);
+					Section.AddMenuEntry(Commands.CenterOfMassRenderingMode_Simulation_Selected);
+					Section.AddMenuEntry(Commands.CenterOfMassRenderingMode_Simulation_None);
+				}
+			}));
+
 			Section.AddSubMenu(TEXT("ConstraintConstraintModeSubMenu"), LOCTEXT("ConstraintConstraintModeSubMenu", "Constraints"), FText::GetEmpty(),
 				FNewToolMenuDelegate::CreateLambda([WeakPhysicsAssetEditor = PhysicsAssetEditor.ToWeakPtr()](UToolMenu* InSubMenu)
 			{
@@ -1313,6 +1338,28 @@ void FPhysicsAssetEditor::BindCommands()
 
 	ViewportCommandList->EndGroup();
 
+	ViewportCommandList->BeginGroup(TEXT("CenterOfMassRenderingMode"));
+
+	ViewportCommandList->MapAction(
+		Commands.CenterOfMassRenderingMode_All,
+		FExecuteAction::CreateSP(this, &FPhysicsAssetEditor::OnCenterOfMassRenderingMode, EPhysicsAssetEditorCenterOfMassViewMode::All, false),
+		FCanExecuteAction(),
+		FIsActionChecked::CreateSP(this, &FPhysicsAssetEditor::IsCenterOfMassRenderingMode, EPhysicsAssetEditorCenterOfMassViewMode::All, false));
+
+	ViewportCommandList->MapAction(
+		Commands.CenterOfMassRenderingMode_Selected,
+		FExecuteAction::CreateSP(this, &FPhysicsAssetEditor::OnCenterOfMassRenderingMode, EPhysicsAssetEditorCenterOfMassViewMode::Selected, false),
+		FCanExecuteAction(),
+		FIsActionChecked::CreateSP(this, &FPhysicsAssetEditor::IsCenterOfMassRenderingMode, EPhysicsAssetEditorCenterOfMassViewMode::Selected, false));
+
+	ViewportCommandList->MapAction(
+		Commands.CenterOfMassRenderingMode_None,
+		FExecuteAction::CreateSP(this, &FPhysicsAssetEditor::OnCenterOfMassRenderingMode, EPhysicsAssetEditorCenterOfMassViewMode::None, false),
+		FCanExecuteAction(),
+		FIsActionChecked::CreateSP(this, &FPhysicsAssetEditor::IsCenterOfMassRenderingMode, EPhysicsAssetEditorCenterOfMassViewMode::None, false));
+
+	ViewportCommandList->EndGroup();
+
 	ViewportCommandList->BeginGroup(TEXT("CollisionRenderingMode"));
 
 	ViewportCommandList->MapAction(
@@ -1385,6 +1432,28 @@ void FPhysicsAssetEditor::BindCommands()
 
 	ViewportCommandList->EndGroup();
 
+	ViewportCommandList->BeginGroup(TEXT("CenterOfMassRenderingMode_Simulation"));
+
+	ViewportCommandList->MapAction(
+		Commands.CenterOfMassRenderingMode_Simulation_All,
+		FExecuteAction::CreateSP(this, &FPhysicsAssetEditor::OnCenterOfMassRenderingMode, EPhysicsAssetEditorCenterOfMassViewMode::All, true),
+		FCanExecuteAction(),
+		FIsActionChecked::CreateSP(this, &FPhysicsAssetEditor::IsCenterOfMassRenderingMode, EPhysicsAssetEditorCenterOfMassViewMode::All, true));
+
+	ViewportCommandList->MapAction(
+		Commands.CenterOfMassRenderingMode_Simulation_Selected,
+		FExecuteAction::CreateSP(this, &FPhysicsAssetEditor::OnCenterOfMassRenderingMode, EPhysicsAssetEditorCenterOfMassViewMode::Selected, true),
+		FCanExecuteAction(),
+		FIsActionChecked::CreateSP(this, &FPhysicsAssetEditor::IsCenterOfMassRenderingMode, EPhysicsAssetEditorCenterOfMassViewMode::Selected, true));
+
+	ViewportCommandList->MapAction(
+		Commands.CenterOfMassRenderingMode_Simulation_None,
+		FExecuteAction::CreateSP(this, &FPhysicsAssetEditor::OnCenterOfMassRenderingMode, EPhysicsAssetEditorCenterOfMassViewMode::None, true),
+		FCanExecuteAction(),
+		FIsActionChecked::CreateSP(this, &FPhysicsAssetEditor::IsCenterOfMassRenderingMode, EPhysicsAssetEditorCenterOfMassViewMode::None, true));
+
+	ViewportCommandList->EndGroup();
+
 	ViewportCommandList->BeginGroup(TEXT("CollisionRenderingMode_Simulation"));
 
 	ViewportCommandList->MapAction(
@@ -1452,6 +1521,12 @@ void FPhysicsAssetEditor::BindCommands()
 		FExecuteAction::CreateSP(this, &FPhysicsAssetEditor::ToggleHideKinematicBodies),
 		FCanExecuteAction(),
 		FIsActionChecked::CreateSP(this, &FPhysicsAssetEditor::IsHidingKinematicBodies));
+	
+	ViewportCommandList->MapAction(
+		Commands.HideBodyMass,
+		FExecuteAction::CreateSP(this, &FPhysicsAssetEditor::ToggleHideBodyMass),
+		FCanExecuteAction(),
+		FIsActionChecked::CreateSP(this, &FPhysicsAssetEditor::IsHidingBodyMass));
 
 	ViewportCommandList->MapAction(
 		Commands.DrawConstraintsAsPoints,
@@ -2706,6 +2781,28 @@ bool FPhysicsAssetEditor::IsMeshRenderingMode(EPhysicsAssetEditorMeshViewMode Mo
 	return Mode == SharedData->GetCurrentMeshViewMode(bSimulation);
 }
 
+void FPhysicsAssetEditor::OnCenterOfMassRenderingMode(EPhysicsAssetEditorCenterOfMassViewMode Mode, bool bSimulation)
+{
+	if (bSimulation)
+	{
+		SharedData->EditorOptions->SimulationCenterOfMassViewMode = Mode;
+	}
+	else
+	{
+		SharedData->EditorOptions->CenterOfMassViewMode = Mode;
+	}
+
+	SharedData->EditorOptions->SaveConfig();
+
+	RefreshPreviewViewport();
+
+}
+
+bool FPhysicsAssetEditor::IsCenterOfMassRenderingMode(EPhysicsAssetEditorCenterOfMassViewMode Mode, bool bSimulation) const
+{
+	return Mode == SharedData->GetCurrentCenterOfMassViewMode(bSimulation);
+}
+
 void FPhysicsAssetEditor::OnCollisionRenderingMode(EPhysicsAssetEditorCollisionViewMode Mode, bool bSimulation)
 {
 	if (bSimulation)
@@ -2794,6 +2891,12 @@ void FPhysicsAssetEditor::ToggleHideKinematicBodies()
 	SharedData->EditorOptions->SaveConfig();
 }
 
+void FPhysicsAssetEditor::ToggleHideBodyMass()
+{
+	SharedData->EditorOptions->bHideBodyMass = !SharedData->EditorOptions->bHideBodyMass;
+	SharedData->EditorOptions->SaveConfig();
+}
+
 bool FPhysicsAssetEditor::IsRenderingOnlySelectedSolid() const
 {
 	return SharedData->EditorOptions->bSolidRenderingForSelectedOnly;
@@ -2809,6 +2912,11 @@ bool FPhysicsAssetEditor::IsHidingKinematicBodies() const
 	return SharedData->EditorOptions->bHideKinematicBodies;
 }
 
+bool FPhysicsAssetEditor::IsHidingBodyMass() const
+{
+	return SharedData->EditorOptions->bHideBodyMass;
+}
+
 bool FPhysicsAssetEditor::IsConstraintRenderingMode(EPhysicsAssetEditorConstraintViewMode Mode, bool bSimulation) const
 {
 	return Mode == SharedData->GetCurrentConstraintViewMode(bSimulation);
@@ -3900,7 +4008,7 @@ void FPhysicsAssetEditor::RefreshFilter()
 {
 	SkeletonTree->RefreshFilter();
 	// make sure we resynchronize the list 
-	HandleViewportSelectionChanged(SharedData->SelectedBodies, SharedData->SelectedConstraints);
+	HandleViewportSelectionChanged(SharedData->SelectedBodies, SharedData->SelectedConstraints, SharedData->SelectedCoMs);
 }
 
 void FPhysicsAssetEditor::HandleCreateNewConstraint(int32 BodyIndex0, int32 BodyIndex1)
diff --git a/Engine/Source/Editor/PhysicsAssetEditor/Private/PhysicsAssetEditor.h b/Engine/Source/Editor/PhysicsAssetEditor/Private/PhysicsAssetEditor.h
--- a/Engine/Source/Editor/PhysicsAssetEditor/Private/PhysicsAssetEditor.h
+++ b/Engine/Source/Editor/PhysicsAssetEditor/Private/PhysicsAssetEditor.h
@@ -73,7 +73,7 @@ public:
 	TSharedPtr<FPhysicsAssetEditorSharedData> GetSharedData() const;
 
 	/** Handles a group selection change... assigns the proper object to the properties widget and the hierarchy tree view */
-	void HandleViewportSelectionChanged(const TArray<FPhysicsAssetEditorSharedData::FSelection>& InSelectedBodies, const TArray<FPhysicsAssetEditorSharedData::FSelection>& InSelectedConstraints);
+	void HandleViewportSelectionChanged(const TArray<FPhysicsAssetEditorSharedData::FSelection>& InSelectedBodies, const TArray<FPhysicsAssetEditorSharedData::FSelection>& InSelectedConstraints, const TArray<FPhysicsAssetEditorSharedData::FSelection>& InSelectedCoMs);
 
 	/** Repopulates the hierarchy tree view */
 	void RefreshHierachyTree();
@@ -238,6 +238,8 @@ private:
 	bool IsToggleSimulation() const;
 	void OnMeshRenderingMode(EPhysicsAssetEditorMeshViewMode Mode, bool bSimulation);
 	bool IsMeshRenderingMode(EPhysicsAssetEditorMeshViewMode Mode, bool bSimulation) const;
+	void OnCenterOfMassRenderingMode(EPhysicsAssetEditorCenterOfMassViewMode Mode, bool bSimulation);
+	bool IsCenterOfMassRenderingMode(EPhysicsAssetEditorCenterOfMassViewMode Mode, bool bSimulation) const;
 	void OnCollisionRenderingMode(EPhysicsAssetEditorCollisionViewMode Mode, bool bSimulation);
 	bool IsCollisionRenderingMode(EPhysicsAssetEditorCollisionViewMode Mode, bool bSimulation) const;
 	void OnConstraintRenderingMode(EPhysicsAssetEditorConstraintViewMode Mode, bool bSimulation);
@@ -251,9 +253,11 @@ private:
 	void ToggleRenderOnlySelectedSolid();
 	void ToggleHideSimulatedBodies();
 	void ToggleHideKinematicBodies();
+	void ToggleHideBodyMass();
 	bool IsRenderingOnlySelectedSolid() const;
 	bool IsHidingSimulatedBodies() const;
 	bool IsHidingKinematicBodies() const;
+	bool IsHidingBodyMass() const;
 	void OnToggleMassProperties();
 	bool IsToggleMassProperties() const;
 	void OnSetCollision(bool bEnable);
diff --git a/Engine/Source/Editor/PhysicsAssetEditor/Private/PhysicsAssetEditorActions.cpp b/Engine/Source/Editor/PhysicsAssetEditor/Private/PhysicsAssetEditorActions.cpp
--- a/Engine/Source/Editor/PhysicsAssetEditor/Private/PhysicsAssetEditorActions.cpp
+++ b/Engine/Source/Editor/PhysicsAssetEditor/Private/PhysicsAssetEditorActions.cpp
@@ -30,6 +30,9 @@ void FPhysicsAssetEditorCommands::RegisterCommands()
 	UI_COMMAND(MeshRenderingMode_Solid, "Solid", "Solid Mesh Rendering Mode (Edit)", EUserInterfaceActionType::RadioButton, FInputChord());
 	UI_COMMAND(MeshRenderingMode_Wireframe, "Wireframe", "Wireframe Mesh Rendering Mode (Edit)", EUserInterfaceActionType::RadioButton, FInputChord());
 	UI_COMMAND(MeshRenderingMode_None, "None", "No Mesh Rendering Mode (Edit)", EUserInterfaceActionType::RadioButton, FInputChord());
+	UI_COMMAND(CenterOfMassRenderingMode_All, "All", "Show Center of Mass for all bodies (Edit)", EUserInterfaceActionType::RadioButton, FInputChord());
+	UI_COMMAND(CenterOfMassRenderingMode_Selected, "Selected", "Show Center of Mass for all selected bodies (Edit)", EUserInterfaceActionType::RadioButton, FInputChord());
+	UI_COMMAND(CenterOfMassRenderingMode_None, "None", "No Center of Mass Rendering (Edit)", EUserInterfaceActionType::RadioButton, FInputChord());
 	UI_COMMAND(CollisionRenderingMode_Solid, "Solid", "Solid Collision Rendering Mode (Edit)", EUserInterfaceActionType::RadioButton, FInputChord());
 	UI_COMMAND(CollisionRenderingMode_Wireframe, "Wireframe", "Wireframe Collision Rendering Mode (Edit)", EUserInterfaceActionType::RadioButton, FInputChord());
 	UI_COMMAND(CollisionRenderingMode_SolidWireframe, "Solid + Wireframe", "Solid + Wireframe Collision Rendering Mode (Edit)", EUserInterfaceActionType::RadioButton, FInputChord());
@@ -41,6 +44,9 @@ void FPhysicsAssetEditorCommands::RegisterCommands()
 	UI_COMMAND(MeshRenderingMode_Simulation_Solid, "Solid", "Solid Mesh Rendering Mode (Simulation)", EUserInterfaceActionType::RadioButton, FInputChord());
 	UI_COMMAND(MeshRenderingMode_Simulation_Wireframe, "Wireframe", "Wireframe Mesh Rendering Mode (Simulation)", EUserInterfaceActionType::RadioButton, FInputChord());
 	UI_COMMAND(MeshRenderingMode_Simulation_None, "None", "No Mesh Rendering Mode (Simulation)", EUserInterfaceActionType::RadioButton, FInputChord());
+	UI_COMMAND(CenterOfMassRenderingMode_Simulation_All, "All", "Show Center of Mass for all bodies (Simulation)", EUserInterfaceActionType::RadioButton, FInputChord());
+	UI_COMMAND(CenterOfMassRenderingMode_Simulation_Selected, "Selected", "Show Center of Mass for all selected bodies (Simulation)", EUserInterfaceActionType::RadioButton, FInputChord());
+	UI_COMMAND(CenterOfMassRenderingMode_Simulation_None, "None", "No Center of Mass Rendering (Simulation)", EUserInterfaceActionType::RadioButton, FInputChord());
 	UI_COMMAND(CollisionRenderingMode_Simulation_Solid, "Solid", "Solid Collision Rendering Mode (Simulation)", EUserInterfaceActionType::RadioButton, FInputChord());
 	UI_COMMAND(CollisionRenderingMode_Simulation_Wireframe, "Wireframe", "Wireframe Collision Rendering Mode (Simulation)", EUserInterfaceActionType::RadioButton, FInputChord());
 	UI_COMMAND(CollisionRenderingMode_Simulation_SolidWireframe, "Solid + Wireframe", "Solid + Wireframe Collision Rendering Mode (Simulation)", EUserInterfaceActionType::RadioButton, FInputChord());
@@ -55,6 +61,7 @@ void FPhysicsAssetEditorCommands::RegisterCommands()
 	UI_COMMAND(RenderOnlySelectedConstraints, "Only Selected Constraints", "Draw only selected constraints.", EUserInterfaceActionType::ToggleButton, FInputChord());
 	UI_COMMAND(DrawConstraintsAsPoints, "Draw Constraints As Points", "Draw Constraints As Points", EUserInterfaceActionType::ToggleButton, FInputChord());
 	UI_COMMAND(DrawViolatedLimits, "Draw Violated Limits", "Draw Violated Limits", EUserInterfaceActionType::ToggleButton, FInputChord());
+	UI_COMMAND(HideBodyMass, "Hide Mass", "Hide text displaying body mass next to CoM", EUserInterfaceActionType::ToggleButton, FInputChord());
 	UI_COMMAND(ToggleMassProperties, "Mass Properties", "Show Mass Properties For Bodies When Simulation Is Enabled", EUserInterfaceActionType::ToggleButton, FInputChord());
 	UI_COMMAND(DisableCollision, "Disable Collision", "Disable collision between the currently selected bodies", EUserInterfaceActionType::Button, FInputChord(EKeys::RightBracket));
 	UI_COMMAND(DisableCollisionAll, "Disable Collision All", "Disable collision between the currently selected bodies and all bodies", EUserInterfaceActionType::Button, FInputChord(EModifierKey::Control, EKeys::RightBracket));
diff --git a/Engine/Source/Editor/PhysicsAssetEditor/Private/PhysicsAssetEditorActions.h b/Engine/Source/Editor/PhysicsAssetEditor/Private/PhysicsAssetEditorActions.h
--- a/Engine/Source/Editor/PhysicsAssetEditor/Private/PhysicsAssetEditorActions.h
+++ b/Engine/Source/Editor/PhysicsAssetEditor/Private/PhysicsAssetEditorActions.h
@@ -44,6 +44,9 @@ public:
 	TSharedPtr<FUICommandInfo> MeshRenderingMode_Solid;
 	TSharedPtr<FUICommandInfo> MeshRenderingMode_Wireframe;
 	TSharedPtr<FUICommandInfo> MeshRenderingMode_None;
+	TSharedPtr<FUICommandInfo> CenterOfMassRenderingMode_All;
+	TSharedPtr<FUICommandInfo> CenterOfMassRenderingMode_Selected;
+	TSharedPtr<FUICommandInfo> CenterOfMassRenderingMode_None;
 	TSharedPtr<FUICommandInfo> CollisionRenderingMode_Solid;
 	TSharedPtr<FUICommandInfo> CollisionRenderingMode_Wireframe;
 	TSharedPtr<FUICommandInfo> CollisionRenderingMode_SolidWireframe;
@@ -54,6 +57,9 @@ public:
 	TSharedPtr<FUICommandInfo> MeshRenderingMode_Simulation_Solid;
 	TSharedPtr<FUICommandInfo> MeshRenderingMode_Simulation_Wireframe;
 	TSharedPtr<FUICommandInfo> MeshRenderingMode_Simulation_None;
+	TSharedPtr<FUICommandInfo> CenterOfMassRenderingMode_Simulation_All;
+	TSharedPtr<FUICommandInfo> CenterOfMassRenderingMode_Simulation_Selected;
+	TSharedPtr<FUICommandInfo> CenterOfMassRenderingMode_Simulation_None;
 	TSharedPtr<FUICommandInfo> CollisionRenderingMode_Simulation_Solid;
 	TSharedPtr<FUICommandInfo> CollisionRenderingMode_Simulation_Wireframe;
 	TSharedPtr<FUICommandInfo> CollisionRenderingMode_Simulation_SolidWireframe;
@@ -64,6 +70,9 @@ public:
 	TSharedPtr<FUICommandInfo> RenderOnlySelectedSolid;
 	TSharedPtr<FUICommandInfo> HideSimulatedBodies;
 	TSharedPtr<FUICommandInfo> HideKinematicBodies;
+	TSharedPtr<FUICommandInfo> HideBodyMass;
+	TSharedPtr<FUICommandInfo> DrawCenterOfMass;
+	TSharedPtr<FUICommandInfo> RenderOnlySelectedBodyCoMs;
 	TSharedPtr<FUICommandInfo> RenderOnlySelectedConstraints;
 	TSharedPtr<FUICommandInfo> DrawConstraintsAsPoints;
 	TSharedPtr<FUICommandInfo> DrawViolatedLimits;
diff --git a/Engine/Source/Editor/PhysicsAssetEditor/Private/PhysicsAssetEditorBodySetupDetailsCustomization.cpp b/Engine/Source/Editor/PhysicsAssetEditor/Private/PhysicsAssetEditorBodySetupDetailsCustomization.cpp
new file mode 100644
--- /dev/null
+++ b/Engine/Source/Editor/PhysicsAssetEditor/Private/PhysicsAssetEditorBodySetupDetailsCustomization.cpp
@@ -0,0 +1,199 @@
+// Copyright Epic Games, Inc. All Rights Reserved.
+
+#include "PhysicsAssetEditorBodySetupDetailsCustomization.h"
+
+#include "DetailCategoryBuilder.h"
+#include "DetailLayoutBuilder.h"
+#include "DetailWidgetRow.h"
+#include "Framework/Application/SlateApplication.h"
+#include "ObjectEditorUtils.h"
+#include "PhysicsAssetEditor.h"
+#include "PhysicsAssetEditorSharedData.h"
+#include "PhysicsAssetEditorSkeletalMeshComponent.h"
+#include "PhysicsEngine/BodySetup.h"
+#include "PhysicsEngine/PhysicsAsset.h"
+#include "PhysicsEngine/SkeletalBodySetup.h"
+#include "PropertyHandle.h"
+#include "UObject/UObjectIterator.h"
+#include "Widgets/Input/SButton.h"
+
+#define LOCTEXT_NAMESPACE "PhysicsAssetBodyDetailsCustomization"
+
+TSharedRef<IDetailCustomization> FPhysicsAssetEditorBodySetupDetailsCustomization::MakeInstance()
+{
+	return MakeShared<FPhysicsAssetEditorBodySetupDetailsCustomization>();
+}
+
+FPhysicsAssetEditorBodySetupDetailsCustomization::FPhysicsAssetEditorBodySetupDetailsCustomization()
+: SharedData(nullptr)
+, BodyIndex(INDEX_NONE)
+{}
+
+void FPhysicsAssetEditorBodySetupDetailsCustomization::CustomizeDetails(IDetailLayoutBuilder& DetailLayout)
+{
+	const UPhysicsAsset* PhysicsAsset = nullptr;
+	BodyIndex = INDEX_NONE;
+
+	// Find the physics asset that owns the BodySetup that is to be customized.
+	{
+		TArray< TWeakObjectPtr<UObject> > ObjectsToBeCustomized;
+		DetailLayout.GetObjectsBeingCustomized(ObjectsToBeCustomized);
+
+		for (TWeakObjectPtr<UObject> ObjectPtr : ObjectsToBeCustomized)
+		{
+			if (ObjectPtr.IsValid())
+			{
+				if (USkeletalBodySetup* const SkeletalBodySetup = Cast<USkeletalBodySetup>(ObjectPtr.Get()))
+				{
+					PhysicsAsset = Cast<UPhysicsAsset>(SkeletalBodySetup->GetOuter());
+
+					if (PhysicsAsset)
+					{
+						BodyIndex = PhysicsAsset->FindBodyIndex(SkeletalBodySetup->BoneName);
+						break;
+					}
+				}
+			}
+		}
+	}
+
+	// Find Physics Editor Shared Data for this Physics Asset.
+	for (TObjectIterator<UPhysicsAssetEditorSkeletalMeshComponent> Itr; Itr; ++Itr)
+	{
+		if (PhysicsAsset && (Itr->SharedData->PhysicsAsset == PhysicsAsset))
+		{
+			SharedData = Itr->SharedData;
+			break;
+		}
+	}
+	
+	FBodySetupDetails::CustomizeDetails(DetailLayout);
+}
+
+void FPhysicsAssetEditorBodySetupDetailsCustomization::CustomizeCoMNudge(IDetailLayoutBuilder& DetailBuilder, TSharedRef<IPropertyHandle> BodyInstanceHandler)
+{
+	IDetailCategoryBuilder& PhysicsCategory = DetailBuilder.EditCategory("Physics");
+
+	TSharedPtr<IPropertyHandle> COMOffsetProperty = BodyInstanceHandler->GetChildHandle(GET_MEMBER_NAME_CHECKED(FBodyInstance, COMNudge));
+
+	TSharedPtr<SHorizontalBox> ValueContent = SNew(SHorizontalBox);
+
+	const FSlateIcon WorldSpaceIcon(FAppStyle::GetAppStyleSetName(), TEXT("EditorViewport.RelativeCoordinateSystem_World"));
+	const FSlateIcon LocalSpaceIcon(FAppStyle::GetAppStyleSetName(), TEXT("Icons.Transform"));
+
+	auto AddComponent = [this, ValueContent, COMOffsetProperty, WorldSpaceIcon, LocalSpaceIcon](const FName ComponentName, const EAxis::Type Axis)
+		{
+			if (TSharedPtr<IPropertyHandle> ComponentHandle = COMOffsetProperty->GetChildHandle(ComponentName))
+			{
+				const float HorizontalPadding = 2.0f;
+				const float VerticalPadding = 2.0f;
+
+				ValueContent->AddSlot()
+					.AutoWidth()
+					.HAlign(HAlign_Left)
+					.VAlign(VAlign_Center)
+					.Padding(HorizontalPadding, VerticalPadding)
+					[
+						ComponentHandle->CreatePropertyNameWidget()
+					];
+
+				ValueContent->AddSlot()
+					.AutoWidth()
+					.HAlign(HAlign_Left)
+					.VAlign(VAlign_Center)
+					.Padding(HorizontalPadding, VerticalPadding)
+					[
+						ComponentHandle->CreatePropertyValueWidget()
+					];
+
+				ValueContent->AddSlot()
+					.AutoWidth()
+					.VAlign(VAlign_Center)
+					[
+						SNew(SButton)
+							.HAlign(HAlign_Center)
+							.VAlign(VAlign_Center)
+							.OnClicked_Lambda([this, Axis]() { return this->ToggleFixCOMInComponentSpace(Axis); })
+							.ButtonColorAndOpacity(FSlateColor::UseForeground())
+							.Content()
+							[
+								SNew(SImage)
+									.ColorAndOpacity(FSlateColor::UseForeground())
+									.Image_Lambda([this, Axis, WorldSpaceIcon, LocalSpaceIcon]() { return (this->IsCOMFixedInComponentSpace(Axis)) ? WorldSpaceIcon.GetIcon() : LocalSpaceIcon.GetIcon(); })
+									.ToolTipText_Lambda([this, Axis]()
+										{ return (this->IsCOMFixedInComponentSpace(Axis))
+										? LOCTEXT("ToolTipDeactivateCOMFixedInComponentSpace", "Center of Mass position is currently fixed in component space on this axis. When the physics body is moved the CoM offset will automatically update to maintain the current position. This will only have an effect in the editor. Click to toggle this behavior (Shift + Click to toggle all axis).")
+										: LOCTEXT("ToolTipActivateCOMFixedInComponentSpace", "Center of Mass position is currently fixed in local space on this axis. When the physics body is moved the CoM offset will move with it to maintain the current position relative to the body. This will only have an effect in the editor. Click to toggle this behavior (Shift + Click to toggle all axis)."); })
+							]
+					];
+			}
+		};
+
+	AddComponent(GET_MEMBER_NAME_CHECKED(FVector, X), EAxis::X);
+	AddComponent(GET_MEMBER_NAME_CHECKED(FVector, Y), EAxis::Y);
+	AddComponent(GET_MEMBER_NAME_CHECKED(FVector, Z), EAxis::Z);
+
+	PhysicsCategory.AddCustomRow(COMOffsetProperty->GetPropertyDisplayName(), true)
+		.NameContent()
+		[
+			COMOffsetProperty->CreatePropertyNameWidget()
+		]
+		.ValueContent()
+		[
+			ValueContent.ToSharedRef()
+		]
+		.ExtensionContent()
+		[
+			SNew(SButton)
+				.IsFocusable(false)
+				.ToolTipText(LOCTEXT("ResetButtonTooltip", "Reset property value to its default value."))
+				.ButtonStyle(FAppStyle::Get(), "SimpleButton")
+				.ContentPadding(0)
+				.OnClicked_Lambda([this, COMOffsetProperty]() { COMOffsetProperty->ResetToDefault(); return FReply::Handled(); })
+				.Content()
+				[
+					SNew(SImage)
+						.Image(FAppStyle::GetBrush("PropertyWindow.DiffersFromDefault"))
+						.ColorAndOpacity(FSlateColor::UseForeground())
+				]
+		];
+	
+	BodyInstanceHandler->GetChildHandle(GET_MEMBER_NAME_CHECKED(FBodyInstance, COMNudge))->MarkHiddenByCustomization();
+
+	COMOffsetProperty->MarkHiddenByCustomization();	
+}
+
+FReply FPhysicsAssetEditorBodySetupDetailsCustomization::ToggleFixCOMInComponentSpace(const EAxis::Type Axis)
+{
+	if (SharedData)
+	{
+		const bool bIsCOMFixed = SharedData->IsCoMAxisFixedInComponentSpace(BodyIndex, Axis);
+
+		if (FSlateApplication::Get().GetModifierKeys().IsShiftDown())
+		{
+			SharedData->SetCoMAxisFixedInComponentSpace(BodyIndex, EAxis::X, !bIsCOMFixed);
+			SharedData->SetCoMAxisFixedInComponentSpace(BodyIndex, EAxis::Y, !bIsCOMFixed);
+			SharedData->SetCoMAxisFixedInComponentSpace(BodyIndex, EAxis::Z, !bIsCOMFixed);
+		}
+		else
+		{
+			SharedData->SetCoMAxisFixedInComponentSpace(BodyIndex, Axis, !bIsCOMFixed);
+		}
+	}
+
+	return FReply::Handled();
+}
+
+bool FPhysicsAssetEditorBodySetupDetailsCustomization::IsCOMFixedInComponentSpace(const EAxis::Type Axis) const
+{
+	bool bIsCOMFixed = false;
+
+	if (SharedData)
+	{
+		bIsCOMFixed = SharedData->IsCoMAxisFixedInComponentSpace(BodyIndex, Axis);;
+	}
+
+	return bIsCOMFixed;
+}
+
+#undef LOCTEXT_NAMESPACE
diff --git a/Engine/Source/Editor/PhysicsAssetEditor/Private/PhysicsAssetEditorBodySetupDetailsCustomization.h b/Engine/Source/Editor/PhysicsAssetEditor/Private/PhysicsAssetEditorBodySetupDetailsCustomization.h
new file mode 100644
--- /dev/null
+++ b/Engine/Source/Editor/PhysicsAssetEditor/Private/PhysicsAssetEditorBodySetupDetailsCustomization.h
@@ -0,0 +1,43 @@
+// Copyright Epic Games, Inc. All Rights Reserved.
+
+#pragma once
+
+#include "CoreMinimal.h"
+#include "IDetailCustomization.h"
+#include "Math/Axis.h"
+
+#include "Editor/DetailCustomizations/Private/BodySetupDetails.h"
+
+class FPhysicsAssetEditor;
+class FPhysicsAssetEditorSharedData;
+class FReply;
+class IDetailCategoryBuilder;
+class IPropertyHandle;
+
+struct FBodyData;
+
+// Class FPhysicsAssetEditorBodySetupDetailsCustomization
+//
+// Replaces the base BodySetupDetails customization, adding support for Center of Mass (CoM) offset
+// editing tools to the details panel.
+//
+class FPhysicsAssetEditorBodySetupDetailsCustomization : public FBodySetupDetails
+{
+public:
+	// Makes a new instance of this detail layout class for a specific detail view requesting it
+	static TSharedRef<IDetailCustomization> MakeInstance();
+
+	FPhysicsAssetEditorBodySetupDetailsCustomization();
+
+	// IDetailCustomization interface
+	virtual void CustomizeDetails(IDetailLayoutBuilder& DetailLayout) override;
+
+	virtual void CustomizeCoMNudge(IDetailLayoutBuilder& DetailBuilder, TSharedRef<IPropertyHandle> BodyInstanceHandler) override;
+	FReply ToggleFixCOMInComponentSpace(const EAxis::Type Axis);
+	bool IsCOMFixedInComponentSpace(const EAxis::Type Axis) const;
+
+private:
+
+	FPhysicsAssetEditorSharedData* SharedData;
+	int32 BodyIndex;
+};
diff --git a/Engine/Source/Editor/PhysicsAssetEditor/Private/PhysicsAssetEditorEditMode.cpp b/Engine/Source/Editor/PhysicsAssetEditor/Private/PhysicsAssetEditorEditMode.cpp
--- a/Engine/Source/Editor/PhysicsAssetEditor/Private/PhysicsAssetEditorEditMode.cpp
+++ b/Engine/Source/Editor/PhysicsAssetEditor/Private/PhysicsAssetEditorEditMode.cpp
@@ -149,12 +149,16 @@ bool FPhysicsAssetEditorEditMode::HandleBeginTransform()
 {
 	if (!SharedData->bRunningSimulation && !SharedData->bManipulating)
 	{
-		if(SharedData->GetSelectedBody() || SharedData->GetSelectedConstraint())
+		if(SharedData->GetSelectedBody() || SharedData->GetSelectedConstraint() || SharedData->GetSelectedCoM())
 		{
 			if(SharedData->GetSelectedBody())
 			{
 				GEditor->BeginTransaction(NSLOCTEXT("UnrealEd", "MoveElement", "Move Element"));
 			}
+			else if (SharedData->GetSelectedCoM())
+			{
+				GEditor->BeginTransaction(NSLOCTEXT("UnrealEd", "MoveCoM", "Move Center of Mass"));
+			}
 			else
 			{
 				GEditor->BeginTransaction(NSLOCTEXT("UnrealEd", "MoveConstraint", "Move Constraint"));
@@ -167,6 +171,20 @@ bool FPhysicsAssetEditorEditMode::HandleBeginTransform()
 			{
 				SharedData->PhysicsAsset->SkeletalBodySetups[SharedData->SelectedBodies[i].Index]->Modify();
 				SharedData->SelectedBodies[i].ManipulateTM = FTransform::Identity;
+				SharedData->SelectedBodies[i].CoMPosition = SharedData->EditorSkelComp->Bodies[SharedData->SelectedBodies[i].Index]->GetCOMPosition();
+			}
+
+			SharedData->bManipulating = true;
+		}
+	
+		if (SharedData->GetSelectedCoM())
+		{
+			for (int32 i = 0; i < SharedData->SelectedCoMs.Num(); ++i)
+			{
+				const int32 BodyIndex = SharedData->SelectedCoMs[i].Index;
+				SharedData->PhysicsAsset->SkeletalBodySetups[BodyIndex]->Modify();
+				SharedData->SelectedCoMs[i].ManipulateTM = FTransform::Identity;
+				SharedData->SelectedCoMs[i].CoMPosition = SharedData->EditorSkelComp->Bodies[BodyIndex]->GetCOMPosition();
 			}
 
 			SharedData->bManipulating = true;
@@ -214,7 +232,6 @@ bool FPhysicsAssetEditorEditMode::HandleEndTransform(FEditorViewportClient* InVi
 
 			FKAggregateGeom* AggGeom = &BodySetup->AggGeom;
 
-
 			if (SelectedObject.PrimitiveType == EAggCollisionShape::Convex)
 			{
 				FKConvexElem& Convex = AggGeom->ConvexElems[SelectedObject.PrimitiveIndex];
@@ -225,6 +242,8 @@ bool FPhysicsAssetEditorEditMode::HandleEndTransform(FEditorViewportClient* InVi
 			}
 		}
 
+		SharedData->bShouldUpdatedSelectedCoMs = true;
+
 		GEditor->EndTransaction();
 		SharedData->RefreshPhysicsAssetChange(SharedData->PhysicsAsset, false);
 
@@ -516,8 +535,40 @@ bool FPhysicsAssetEditorEditMode::InputDelta(FEditorViewportClient* InViewportCl
 				}
 			}
 		}
-	}
+
+		for (int32 i = 0; i < SharedData->SelectedCoMs.Num(); ++i)
+		{
+			FPhysicsAssetEditorSharedData::FSelection& SelectedObject = SharedData->SelectedCoMs[i];
+			if (SharedData->bManipulating)
+			{
+				const int32 BodyIndex = SelectedObject.Index;
+				const int32 BoneIndex = SharedData->EditorSkelComp->GetBoneIndex(SharedData->PhysicsAsset->SkeletalBodySetups[BodyIndex]->BoneName);
+				FTransform BoneTM = SharedData->EditorSkelComp->GetBoneTransform(BoneIndex);
+				const float BoneScale = BoneTM.GetScale3D().GetAbsMax();
+				BoneTM.RemoveScaling();
+
+				SelectedObject.WidgetTM.SetRotation(BoneTM.GetRotation());
+				SelectedObject.WidgetTM.SetTranslation(SelectedObject.CoMPosition);
 	
+				if (InViewportClient->GetWidgetMode() == UE::Widget::WM_Translate)
+				{
+					FVector DragToUse = GetLocalTranslation(InViewportClient, InDrag, SelectedObject.WidgetTM);
+					FVector Dir = DragToUse.GetSafeNormal();
+					FVector DragVec = Dir * DragToUse.Size() / BoneScale;
+
+					SelectedObject.CoMPosition += DragVec; // Use widget movement to update selected objects target worldspace location, then calculate the correct CoM nudge to position the world space CoM at this location when manipulation ends.
+
+					// Update CoM nudge value in body instance to ensure the values in the details panel update as the manipulator is moved in the viewport.
+					const FVector CalculatedCoMOffset = CalculateCoMNudgeForWorldSpacePosition(BodyIndex, SelectedObject.CoMPosition);
+					SharedData->EditorSkelComp->Bodies[BodyIndex]->COMNudge = CalculatedCoMOffset;
+					SharedData->PhysicsAsset->SkeletalBodySetups[BodyIndex]->DefaultInstance.COMNudge = CalculatedCoMOffset;
+				}
+			}
+
+			bHandled = true;
+		}
+	}
+
 	return bHandled;
 }
 
@@ -568,6 +619,8 @@ void FPhysicsAssetEditorEditMode::Tick(FEditorViewportClient* ViewportClient, fl
 			EnumRemoveFlags(RenderSettings->ConstraintViewportManipulationFlags, EConstraintTransformComponentFlags::AllChild); // Remove Child Frame flags.
 		}	
 	}
+
+	UpdateCoM();
 }
 
 void FPhysicsAssetEditorEditMode::Render(const FSceneView* View, FViewport* Viewport, FPrimitiveDrawInterface* PDI)
@@ -643,27 +696,35 @@ void FPhysicsAssetEditorEditMode::DrawHUD(FEditorViewportClient* ViewportClient,
 	int32 HalfY = Viewport->GetSizeXY().Y / 2;
 
 	// If showing center-of-mass, and physics is started up..
-	if (SharedData->GetShowCom())
+	if (FPhysicsAssetRenderSettings* const RenderSettings = UPhysicsAssetRenderUtilities::GetSettings(SharedData->PhysicsAsset))
 	{
-		// iterate over each bone
-		for (int32 i = 0; i <SharedData->EditorSkelComp->Bodies.Num(); ++i)
+		if (RenderSettings && !RenderSettings->bHideBodyMass && (RenderSettings->CenterOfMassViewMode != EPhysicsAssetEditorCenterOfMassViewMode::None))
 		{
-			FBodyInstance* BodyInst = SharedData->EditorSkelComp->Bodies[i];
-			check(BodyInst);
+			const bool bDrawSelectedOnly = RenderSettings->CenterOfMassViewMode == EPhysicsAssetEditorCenterOfMassViewMode::Selected;
+			
+			// iterate over each body
+			for (int32 BodyIndex = 0, BodyCount = SharedData->EditorSkelComp->Bodies.Num(); BodyIndex < BodyCount; ++BodyIndex)
+			{
+				if (!bDrawSelectedOnly || SharedData->IsBodySelected(BodyIndex))
+				{
+					FBodyInstance* BodyInst = SharedData->EditorSkelComp->Bodies[BodyIndex];
+					check(BodyInst);
 
-			FVector BodyCOMPos = BodyInst->GetCOMPosition();
-			float BodyMass = BodyInst->GetBodyMass();
+					FVector BodyCOMPos = SharedData->GetCOMRenderPosition(BodyIndex);
+					float BodyMass = BodyInst->GetBodyMass();
 
-			FPlane Projection = View->Project(BodyCOMPos);
-			if (Projection.W > 0.f) // This avoids drawing bone names that are behind us.
-			{
-				int32 XPos = HalfX + (HalfX * Projection.X);
-				int32 YPos = HalfY + (HalfY * (Projection.Y * -1));
+					FPlane Projection = View->Project(BodyCOMPos);
+					if (Projection.W > 0.f) // This avoids drawing bone names that are behind us.
+					{
+						const int32 XPos = HalfX + (HalfX * Projection.X) + RenderSettings->COMRenderMassTextOffsetScreenspace;
+						const int32 YPos = HalfY + (HalfY * (Projection.Y * -1)) + RenderSettings->COMRenderMassTextOffsetScreenspace;
 
-				FString COMString = FString::Printf(TEXT("%3.3f"), BodyMass);
-				TextItem.Text = FText::FromString(COMString);
-				TextItem.SetColor(SharedData->COMRenderColor);
-				Canvas->DrawItem(TextItem, XPos, YPos);
+						FString COMString = FString::Printf(TEXT("%3.3f"), BodyMass);
+						TextItem.Text = FText::FromString(COMString);
+						TextItem.SetColor(SharedData->COMRenderColor);
+						Canvas->DrawItem(TextItem, XPos, YPos);
+					}
+				}
 			}
 		}
 	}
@@ -676,7 +737,7 @@ bool FPhysicsAssetEditorEditMode::AllowWidgetMove()
 
 bool FPhysicsAssetEditorEditMode::ShouldDrawWidget() const
 {
-	return !SharedData->bRunningSimulation && (SharedData->GetSelectedBody() || SharedData->GetSelectedConstraint());
+	return !SharedData->bRunningSimulation && (SharedData->GetSelectedBody() || SharedData->GetSelectedConstraint() || SharedData->GetSelectedCoM());
 }
 
 bool FPhysicsAssetEditorEditMode::UsesTransformWidget() const
@@ -691,6 +752,11 @@ bool FPhysicsAssetEditorEditMode::UsesTransformWidget(UE::Widget::EWidgetMode Ch
 		return false;
 	}
 
+	if (SharedData->GetSelectedCoM() && (CheckMode != UE::Widget::WM_Translate))
+	{
+		return false;
+	}
+
 	return ShouldDrawWidget() && (CheckMode == UE::Widget::WM_Scale || CheckMode == UE::Widget::WM_Translate || CheckMode == UE::Widget::WM_Rotate || CheckMode == UE::Widget::WM_None);
 }
 
@@ -724,6 +790,18 @@ bool FPhysicsAssetEditorEditMode::GetCustomDrawingCoordinateSystem(FMatrix& InMa
 			return true;
 		}
 	}
+	else if (SharedData->GetSelectedCoM())
+	{
+		const int32 BoneIndex = SharedData->EditorSkelComp->GetBoneIndex(SharedData->PhysicsAsset->SkeletalBodySetups[SharedData->GetSelectedCoM()->Index]->BoneName);
+		if (BoneIndex != INDEX_NONE)
+		{
+			FTransform SelectedCoMTM = SharedData->EditorSkelComp->GetBoneTransform(BoneIndex);
+			SelectedCoMTM.RemoveScaling();
+
+			InMatrix = SelectedCoMTM.ToMatrixNoScale().RemoveTranslation();
+			return true;
+		}
+	}
 	else if (SharedData->GetSelectedConstraint())
 	{
 		InMatrix = SharedData->GetConstraintMatrix(SharedData->GetSelectedConstraint()->Index, GetConstraintFrameForWidget(), 1.f).ToMatrixNoScale().RemoveTranslation();
@@ -753,6 +831,18 @@ FVector FPhysicsAssetEditorEditMode::GetWidgetLocation() const
 			return SharedData->EditorSkelComp->GetPrimitiveTransform(BoneTM, SharedData->GetSelectedBody()->Index, SharedData->GetSelectedBody()->PrimitiveType, SharedData->GetSelectedBody()->PrimitiveIndex, Scale).GetTranslation();
 		}
 	}
+	else if (SharedData->GetSelectedCoM())
+	{
+		if (SharedData->bManipulating)
+		{
+			// return the CoM position from the FSelection object because the physics body's CoM position will only be updated at the end of manipulation.
+			return SharedData->GetSelectedCoM()->CoMPosition;
+		}
+		else
+		{
+			return SharedData->EditorSkelComp->Bodies[SharedData->GetSelectedCoM()->Index]->GetCOMPosition();
+		}
+	}
 	else if (SharedData->GetSelectedConstraint())
 	{
 		return SharedData->GetConstraintMatrix(SharedData->GetSelectedConstraint()->Index, GetConstraintFrameForWidget(), 1.f).GetTranslation();
@@ -781,6 +871,13 @@ bool FPhysicsAssetEditorEditMode::HandleClick(FEditorViewportClient* InViewportC
 				SharedData->HitConstraint(ConstraintProxy->ConstraintIndex, InViewportClient->IsCtrlPressed() || InViewportClient->IsShiftPressed());
 				return true;
 			}
+			else if (HitProxy && HitProxy->IsA(HPhysicsAssetEditorEdCoMProxy::StaticGetType()))
+			{
+				HPhysicsAssetEditorEdCoMProxy* CoMProxy = (HPhysicsAssetEditorEdCoMProxy*)HitProxy;
+
+				SharedData->HitCoM(CoMProxy->BodyIndex, InViewportClient->IsCtrlPressed() || InViewportClient->IsShiftPressed());
+				return true;
+			}
 			else
 			{
 				HitNothing(InViewportClient);
@@ -1127,4 +1224,67 @@ void FPhysicsAssetEditorEditMode::HitNothing(FEditorViewportClient* InViewportCl
 	PhysicsAssetEditorPtr.Pin()->RefreshHierachyTree();
 }
 
+void FPhysicsAssetEditorEditMode::UpdateCoM()
+{
+	if (SharedData->bShouldUpdatedSelectedCoMs) // < This calculation must be delayed by a frame s.t. changes to the physics state have been propagated to the physics bodies.
+	{
+		TSet<int32> EncounteredBodyIndexs;
+
+		for (int32 i = 0; i < SharedData->SelectedBodies.Num(); ++i)
+		{
+			FPhysicsAssetEditorSharedData::FSelection& SelectedObject = SharedData->SelectedBodies[i];
+			const int32 BodyIndex = SelectedObject.Index;
+			FBodyInstance* const EditorBodyInstance = SharedData->EditorSkelComp->Bodies[BodyIndex];
+
+			if (SharedData->IsCoMAxisFixedInComponentSpace(BodyIndex, EAxis::X) || SharedData->IsCoMAxisFixedInComponentSpace(BodyIndex, EAxis::Y) || SharedData->IsCoMAxisFixedInComponentSpace(BodyIndex, EAxis::Z))
+			{
+				if (!EncounteredBodyIndexs.Contains(BodyIndex)) // Ensure we only update the CoM of each physics body once (SelectedBodies actually contains a list of selected primitives).
+				{
+					EncounteredBodyIndexs.Add(BodyIndex);
+
+					const FVector CoMOffset = EditorBodyInstance->COMNudge;
+
+					FVector	CalculatedCoMOffset = CalculateCoMNudgeForWorldSpacePosition(BodyIndex, SelectedObject.CoMPosition);
+
+					// Only apply lock to the specified Axis in bone space.
+					if (!SharedData->IsCoMAxisFixedInComponentSpace(BodyIndex, EAxis::X)) { CalculatedCoMOffset.X = CoMOffset.X; }
+					if (!SharedData->IsCoMAxisFixedInComponentSpace(BodyIndex, EAxis::Y)) { CalculatedCoMOffset.Y = CoMOffset.Y; }
+					if (!SharedData->IsCoMAxisFixedInComponentSpace(BodyIndex, EAxis::Z)) { CalculatedCoMOffset.Z = CoMOffset.Z; }
+
+					SharedData->PhysicsAsset->SkeletalBodySetups[BodyIndex]->DefaultInstance.COMNudge = CalculatedCoMOffset;
+				}
+			}
+		}
+
+		for (int32 i = 0; i < SharedData->SelectedCoMs.Num(); ++i)
+		{
+			FPhysicsAssetEditorSharedData::FSelection& SelectedObject = SharedData->SelectedCoMs[i];
+			const FVector CalculatedCoMOffset = CalculateCoMNudgeForWorldSpacePosition(SelectedObject.Index, SelectedObject.CoMPosition);
+			SharedData->PhysicsAsset->SkeletalBodySetups[SelectedObject.Index]->DefaultInstance.COMNudge = CalculatedCoMOffset;
+		}
+
+		SharedData->RefreshPhysicsAssetChange(SharedData->PhysicsAsset, false);
+
+		SharedData->bShouldUpdatedSelectedCoMs = false;
+	}
+}
+
+FVector FPhysicsAssetEditorEditMode::CalculateCoMNudgeForWorldSpacePosition(const int32 BodyIndex, const FVector& CoMPositionWorldSpace) const
+{
+	FVector	CalculatedCoMOffset;
+
+	if (SharedData->EditorSkelComp->Bodies.IsValidIndex(BodyIndex))
+	{
+		if (FBodyInstance* const EditorBodyInstance = SharedData->EditorSkelComp->Bodies[BodyIndex])
+		{
+			const int32 BoneIndex = SharedData->EditorSkelComp->GetBoneIndex(SharedData->PhysicsAsset->SkeletalBodySetups[BodyIndex]->BoneName);
+			const FTransform BoneTM = SharedData->EditorSkelComp->GetBoneTransform(BoneIndex);
+			const FVector CoMWithoutNudge = EditorBodyInstance->GetMassSpaceLocal().GetTranslation() - EditorBodyInstance->COMNudge;
+			CalculatedCoMOffset = BoneTM.InverseTransformPosition(CoMPositionWorldSpace) - CoMWithoutNudge;
+		}
+	}
+
+	return CalculatedCoMOffset;
+}
+
 #undef LOCTEXT_NAMESPACE
diff --git a/Engine/Source/Editor/PhysicsAssetEditor/Private/PhysicsAssetEditorEditMode.h b/Engine/Source/Editor/PhysicsAssetEditor/Private/PhysicsAssetEditorEditMode.h
--- a/Engine/Source/Editor/PhysicsAssetEditor/Private/PhysicsAssetEditorEditMode.h
+++ b/Engine/Source/Editor/PhysicsAssetEditor/Private/PhysicsAssetEditorEditMode.h
@@ -90,6 +90,12 @@ private:
 	bool HandleBeginTransform();
 	bool HandleEndTransform(FEditorViewportClient* InViewportClient) const;
 
+	// Update the Center of Mass position after body transforms have been manipulated in the view port.
+	void UpdateCoM();
+
+	// Calculate a Center of Mass nudge (offset) for a given body that will locate that body's CoM at the supplied position in world space.
+	FVector CalculateCoMNudgeForWorldSpacePosition(const int32 BodyIndex, const FVector& CoMPositionWorldSpace) const;
+
 private:
 	/** Shared data */
 	FPhysicsAssetEditorSharedData* SharedData;
diff --git a/Engine/Source/Editor/PhysicsAssetEditor/Private/PhysicsAssetEditorHitProxies.cpp b/Engine/Source/Editor/PhysicsAssetEditor/Private/PhysicsAssetEditorHitProxies.cpp
--- a/Engine/Source/Editor/PhysicsAssetEditor/Private/PhysicsAssetEditorHitProxies.cpp
+++ b/Engine/Source/Editor/PhysicsAssetEditor/Private/PhysicsAssetEditorHitProxies.cpp
@@ -3,5 +3,6 @@
 #include "PhysicsAssetEditorHitProxies.h"
 
 IMPLEMENT_HIT_PROXY(HPhysicsAssetEditorEdBoneProxy,HHitProxy);
+IMPLEMENT_HIT_PROXY(HPhysicsAssetEditorEdCoMProxy, HHitProxy);
 IMPLEMENT_HIT_PROXY(HPhysicsAssetEditorEdConstraintProxy,HHitProxy);
 IMPLEMENT_HIT_PROXY(HPhysicsAssetEditorEdBoneNameProxy,HHitProxy);
diff --git a/Engine/Source/Editor/PhysicsAssetEditor/Private/PhysicsAssetEditorHitProxies.h b/Engine/Source/Editor/PhysicsAssetEditor/Private/PhysicsAssetEditorHitProxies.h
--- a/Engine/Source/Editor/PhysicsAssetEditor/Private/PhysicsAssetEditorHitProxies.h
+++ b/Engine/Source/Editor/PhysicsAssetEditor/Private/PhysicsAssetEditorHitProxies.h
@@ -36,6 +36,28 @@ struct HPhysicsAssetEditorEdBoneProxy : public HHitProxy
 	}
 };
 
+struct HPhysicsAssetEditorEdCoMProxy : public HHitProxy
+{
+	DECLARE_HIT_PROXY();
+
+	int32							BodyIndex;
+
+	HPhysicsAssetEditorEdCoMProxy(const int32 InBodyIndex)
+		: HHitProxy(HPP_World)
+		, BodyIndex(InBodyIndex)
+	{}
+
+	virtual EMouseCursor::Type GetMouseCursor()
+	{
+		return EMouseCursor::Crosshairs;
+	}
+
+	virtual bool AlwaysAllowsTranslucentPrimitives() const override
+	{
+		return true;
+	}
+};
+
 struct HPhysicsAssetEditorEdConstraintProxy : public HHitProxy
 {
 	DECLARE_HIT_PROXY();
diff --git a/Engine/Source/Editor/PhysicsAssetEditor/Private/PhysicsAssetEditorModule.cpp b/Engine/Source/Editor/PhysicsAssetEditor/Private/PhysicsAssetEditorModule.cpp
--- a/Engine/Source/Editor/PhysicsAssetEditor/Private/PhysicsAssetEditorModule.cpp
+++ b/Engine/Source/Editor/PhysicsAssetEditor/Private/PhysicsAssetEditorModule.cpp
@@ -1,13 +1,15 @@
 // Copyright Epic Games, Inc. All Rights Reserved.
 
 #include "PhysicsAssetEditorModule.h"
-#include "Modules/ModuleManager.h"
-#include "PhysicsAssetEditorSharedData.h"
+#include "EdGraphUtilities.h"
 #include "IPhysicsAssetEditor.h"
+#include "Modules/ModuleManager.h"
+#include "PhysicsAssetEditorBodySetupDetailsCustomization.h"
 #include "PhysicsAssetEditor.h"
-#include "PhysicsAssetGraph/PhysicsAssetGraphPanelNodeFactory.h"
-#include "EdGraphUtilities.h"
 #include "PhysicsAssetEditorEditMode.h"
+#include "PhysicsAssetEditorSharedData.h"
+#include "PhysicsAssetGraph/PhysicsAssetGraphPanelNodeFactory.h"
+#include "PropertyEditorModule.h"
 
 #define LOCTEXT_NAMESPACE "PhysicsAssetEditorModule"
 
@@ -33,6 +35,10 @@ public:
 		FEdGraphUtilities::RegisterVisualNodeFactory(PhysicsAssetGraphPanelNodeFactory);
 
 		FEditorModeRegistry::Get().RegisterMode<FPhysicsAssetEditorEditMode>(FPhysicsAssetEditorEditMode::ModeName, LOCTEXT("PhysicsAssetEditorEditMode", "Physics Asset Editor"), FSlateIcon(), false);
+
+		// Register details customizations.
+		FPropertyEditorModule& PropertyEditorModule = FModuleManager::LoadModuleChecked<FPropertyEditorModule>("PropertyEditor");
+		PropertyEditorModule.RegisterCustomClassLayout("BodySetup", FOnGetDetailCustomizationInstance::CreateStatic(&FPhysicsAssetEditorBodySetupDetailsCustomization::MakeInstance));
 	}
 
 	/** Called before the module is unloaded, right before the module object is destroyed. */
diff --git a/Engine/Source/Editor/PhysicsAssetEditor/Private/PhysicsAssetEditorSharedData.cpp b/Engine/Source/Editor/PhysicsAssetEditor/Private/PhysicsAssetEditorSharedData.cpp
--- a/Engine/Source/Editor/PhysicsAssetEditor/Private/PhysicsAssetEditorSharedData.cpp
+++ b/Engine/Source/Editor/PhysicsAssetEditor/Private/PhysicsAssetEditorSharedData.cpp
@@ -4,6 +4,7 @@
 #include "PhysicsAssetEditorPhysicsHandleComponent.h"
 #include "PhysicsAssetRenderUtils.h"
 #include "PhysicsEngine/RigidBodyIndexPair.h"
+#include "Math/Axis.h"
 #include "Misc/MessageDialog.h"
 #include "Modules/ModuleManager.h"
 #include "UObject/UObjectIterator.h"
@@ -46,14 +47,49 @@
 
 #define LOCTEXT_NAMESPACE "PhysicsAssetEditorShared"
 
-//PRAGMA_DISABLE_OPTIMIZATION
-
 namespace SharedDataConstants
 {
 	const FString ConstraintType = TEXT("Constraint");
 	const FString BodyType = TEXT("SkeletalBodySetup");
 }
 
+bool SetSelected(TArray<FPhysicsAssetEditorSharedData::FSelection>& SelectedElements, const TArray<FPhysicsAssetEditorSharedData::FSelection>& DeltaElements, const bool bSelected)
+{
+	if (DeltaElements.Num() == 0)
+	{
+		return false;
+	}
+
+	if (bSelected)
+	{
+		for (const FPhysicsAssetEditorSharedData::FSelection& Element : DeltaElements)
+		{
+			SelectedElements.AddUnique(Element);
+		}
+	}
+	else
+	{
+		for (const FPhysicsAssetEditorSharedData::FSelection& Element : DeltaElements)
+		{
+			SelectedElements.Remove(Element);
+		}
+	}
+
+	return true;
+}
+
+bool SetSelected(TArray<FPhysicsAssetEditorSharedData::FSelection>& SelectedElements, const TArray<int32>& DeltaElementIndexs, const bool bSelected)
+{
+	TArray<FPhysicsAssetEditorSharedData::FSelection> DeltaElements;
+	DeltaElements.Reserve(DeltaElementIndexs.Num());
+
+	for (const int32 ElementIndex : DeltaElementIndexs)
+	{
+		DeltaElements.Add(FPhysicsAssetEditorSharedData::FSelection(ElementIndex, EAggCollisionShape::Unknown, INDEX_NONE));
+	}
+
+	return SetSelected(SelectedElements, DeltaElements, bSelected);
+}
 
 FScopedBulkSelection::FScopedBulkSelection(TSharedPtr<FPhysicsAssetEditorSharedData> InSharedData)
 	: SharedData(InSharedData)
@@ -67,6 +103,30 @@ FScopedBulkSelection::~FScopedBulkSelection()
 	SharedData->BroadcastSelectionChanged();
 }
 
+// struct FBodyData //
+
+FBodyData::FBodyData()
+: CoMAxisFixedInComponentSpaceFlags(0)
+{}
+
+bool FBodyData::IsCoMAxisFixedInComponentSpace(const EAxis::Type InAxis) const
+{
+	return CoMAxisFixedInComponentSpaceFlags & GetBitFlag(InAxis);
+}
+
+void FBodyData::SetCoMAxisFixedInComponentSpace(const EAxis::Type InAxis, const bool bValue)
+{
+	CoMAxisFixedInComponentSpaceFlags = (bValue) ? CoMAxisFixedInComponentSpaceFlags | GetBitFlag(InAxis) : CoMAxisFixedInComponentSpaceFlags & ~GetBitFlag(InAxis);
+}
+
+int32 FBodyData::GetBitFlag(const EAxis::Type InAxis) const
+{
+	return int32(1) << InAxis;
+}
+
+
+// class FPhysicsAssetEditorSharedData //
+
 FPhysicsAssetEditorSharedData::FPhysicsAssetEditorSharedData()
 	: COMRenderColor(255,255,100)
 	, bSuspendSelectionBroadcast(false)
@@ -184,10 +244,13 @@ void FPhysicsAssetEditorSharedData::Initialize(const TSharedRef<IPersonaPreviewS
 		}
 	}
 
+	EditorBodyData.SetNum(PhysicsAsset->SkeletalBodySetups.Num(), EAllowShrinking::Yes);
+
 	// Support undo/redo
 	PhysicsAsset->SetFlags(RF_Transactional);
 
 	ClearSelectedBody();
+	ClearSelectedCoMs();
 	ClearSelectedConstraints();
 }
 
@@ -195,7 +258,7 @@ void FPhysicsAssetEditorSharedData::BroadcastSelectionChanged()
 {
 	if (!bSuspendSelectionBroadcast)
 	{
-		SelectionChangedEvent.Broadcast(SelectedBodies, SelectedConstraints);
+		SelectionChangedEvent.Broadcast(SelectedBodies, SelectedConstraints, SelectedCoMs);
 	}
 }
 
@@ -290,17 +353,14 @@ bool FPhysicsAssetEditorSharedData::ClipboardHasCompatibleData()
 
 void FPhysicsAssetEditorSharedData::ToggleShowCom()
 {
-	if(FPhysicsAssetRenderSettings* const PhysicsAssetRenderSettings = GetRenderSettings())
-	{
-		PhysicsAssetRenderSettings->bShowCOM = !PhysicsAssetRenderSettings->bShowCOM;
-	}
+	SetShowCom(!GetShowCom());
 }
 
 void FPhysicsAssetEditorSharedData::SetShowCom(bool InValue)
 {
 	if(FPhysicsAssetRenderSettings* const PhysicsAssetRenderSettings = GetRenderSettings())
 	{
-		PhysicsAssetRenderSettings->bShowCOM = InValue;
+		PhysicsAssetRenderSettings->CenterOfMassViewMode = (InValue) ? EPhysicsAssetEditorCenterOfMassViewMode::All : EPhysicsAssetEditorCenterOfMassViewMode::None;
 	}
 }
 
@@ -308,12 +368,53 @@ bool FPhysicsAssetEditorSharedData::GetShowCom() const
 {
 	if(FPhysicsAssetRenderSettings* const PhysicsAssetRenderSettings = GetRenderSettings())
 	{
-		return PhysicsAssetRenderSettings->bShowCOM;
+		return PhysicsAssetRenderSettings->CenterOfMassViewMode == EPhysicsAssetEditorCenterOfMassViewMode::All;
 	}
 
 	return false;
 }
 
+FVector FPhysicsAssetEditorSharedData::GetCOMRenderPosition(const int32 BodyIndex) const
+{
+	if (bManipulating) 
+	{
+		if (const FSelection* const SelectedCoM = SelectedCoMs.FindByPredicate([BodyIndex](const FSelection& Element) { return Element.Index == BodyIndex; }))
+		{
+			// Return the Selection objects CoM position when manipulating as that is the one we're actually updating with the 
+			// manipulator widget (as updating the CoM in the physics body proper is complicated).
+			return SelectedCoM->CoMPosition;
+		}
+	}
+
+	if (EditorSkelComp && EditorSkelComp->Bodies.IsValidIndex(BodyIndex))
+	{
+		if (const FBodyInstance* const EditorBodyInstance = EditorSkelComp->Bodies[BodyIndex])
+		{
+			return EditorBodyInstance->GetCOMPosition();
+		}
+	}
+
+	return FVector::ZeroVector;
+}
+
+bool FPhysicsAssetEditorSharedData::IsCoMAxisFixedInComponentSpace(const int32 BodyIndex, const EAxis::Type InAxis) const
+{
+	if (EditorBodyData.IsValidIndex(BodyIndex))
+	{
+		return EditorBodyData[BodyIndex].IsCoMAxisFixedInComponentSpace(InAxis);
+	}
+
+	return false;
+}
+
+void FPhysicsAssetEditorSharedData::SetCoMAxisFixedInComponentSpace(const int32 BodyIndex, const EAxis::Type InAxis, const bool bValue)
+{
+	if (EditorBodyData.IsValidIndex(BodyIndex))
+	{
+		return EditorBodyData[BodyIndex].SetCoMAxisFixedInComponentSpace(InAxis, bValue);
+	}
+}
+
 bool FPhysicsAssetEditorSharedData::ParseClipboard(UPhysicsAsset*& OutAsset, FString& OutObjectType, UObject*& OutObject)
 {
 	FString ClipboardContent;
@@ -550,6 +651,18 @@ EPhysicsAssetEditorMeshViewMode FPhysicsAssetEditorSharedData::GetCurrentMeshVie
 	}
 }
 
+EPhysicsAssetEditorCenterOfMassViewMode FPhysicsAssetEditorSharedData::GetCurrentCenterOfMassViewMode(const bool bSimulation) const
+{
+	if (bSimulation)
+	{
+		return EditorOptions->SimulationCenterOfMassViewMode;
+	}
+	else
+	{
+		return EditorOptions->CenterOfMassViewMode;
+	}
+}
+
 EPhysicsAssetEditorCollisionViewMode FPhysicsAssetEditorSharedData::GetCurrentCollisionViewMode(bool bSimulation)
 {
 	if (bSimulation)
@@ -579,6 +692,12 @@ void FPhysicsAssetEditorSharedData::HitBone(int32 BodyIndex, EAggCollisionShape:
 	if (!bRunningSimulation)
 	{
 		FPhysicsAssetEditorSharedData::FSelection Selection(BodyIndex, PrimType, PrimIndex);
+
+		if (EditorSkelComp && EditorSkelComp->Bodies.IsValidIndex(BodyIndex))
+		{
+			Selection.CoMPosition = EditorSkelComp->Bodies[BodyIndex]->GetCOMPosition();
+		}
+
 		if(bGroupSelect)
 		{
 			if(IsBodySelected(Selection))
@@ -598,6 +717,36 @@ void FPhysicsAssetEditorSharedData::HitBone(int32 BodyIndex, EAggCollisionShape:
 	}
 }
 
+void FPhysicsAssetEditorSharedData::HitCoM(const int32 BodyIndex, const bool bGroupSelect)
+{
+	if (!bRunningSimulation)
+	{
+		FPhysicsAssetEditorSharedData::FSelection Selection(BodyIndex, EAggCollisionShape::Unknown, INDEX_NONE);
+
+		if (EditorSkelComp && EditorSkelComp->Bodies.IsValidIndex(BodyIndex))
+		{
+			Selection.CoMPosition = EditorSkelComp->Bodies[BodyIndex]->GetCOMPosition();
+		}
+
+		if (bGroupSelect)
+		{
+			if (IsCoMSelected(BodyIndex))
+			{
+				SetSelectedCoM(Selection, false);
+			}
+			else
+			{
+				SetSelectedCoM(Selection, true);
+			}
+		}
+		else
+		{
+			ClearSelectedCoMs();
+			SetSelectedCoM(Selection, true);
+		}
+	}
+}
+
 void FPhysicsAssetEditorSharedData::HitConstraint(int32 ConstraintIndex, bool bGroupSelect)
 {
 	if (!bRunningSimulation)
@@ -627,14 +776,20 @@ void FPhysicsAssetEditorSharedData::RefreshPhysicsAssetChange(const UPhysicsAsse
 	{
 		InPhysAsset->RefreshPhysicsAssetChange();
 
-		// Broadbcast delegate
+		// Broadcast delegate
 		FPhysicsDelegates::OnPhysicsAssetChanged.Broadcast(InPhysAsset);
 
 		FEditorSupportDelegates::RedrawAllViewports.Broadcast();
-		// since we recreate physicsstate, a lot of transient state data will be gone
+		// since we recreate physics state, a lot of transient state data will be gone
 		// so have to turn simulation off again. 
 		// ideally maybe in the future, we'll fix it by controlling tick?
 		EditorSkelComp->RecreatePhysicsState();
+
+		for (int32 BodyIndex = 0, BodyCount = EditorSkelComp->Bodies.Num(); BodyIndex < BodyCount; ++BodyIndex)
+		{
+			EditorSkelComp->Bodies[BodyIndex]->BodySetup = InPhysAsset->SkeletalBodySetups[BodyIndex];
+		}
+
 		if(bFullClothRefresh)
 		{
 			EditorSkelComp->RecreateClothingActors();
@@ -725,13 +880,56 @@ void FPhysicsAssetEditorSharedData::SetSelectedBodiesPrimitives(const TArray<int
 	SetSelectedBodies(NewSelection, bSelected);
 }
 
-void FPhysicsAssetEditorSharedData::ClearSelectedBody()
+void FPhysicsAssetEditorSharedData::ClearSelected()
 {
 	SelectedBodies.Empty();
+	SelectedCoMs.Empty();
 	SelectedConstraints.Empty();
+
 	BroadcastSelectionChanged();
 }
 
+void FPhysicsAssetEditorSharedData::ClearSelectedCoMs()
+{
+	if (InsideSelChange)
+	{
+		return;
+	}
+
+	ClearSelected();
+
+	++InsideSelChange;
+	BroadcastPreviewChanged();
+	--InsideSelChange;
+}
+
+void FPhysicsAssetEditorSharedData::SetSelectedCoM(const FSelection& InSelectedElement, const bool bSelected)
+{
+	SetSelectedCoMs({ InSelectedElement }, bSelected);
+}
+
+void FPhysicsAssetEditorSharedData::SetSelectedCoMs(const TArray<FSelection>& InSelectedElements, const bool bSelected)
+{
+	if (!InsideSelChange && SetSelected(SelectedCoMs, InSelectedElements, bSelected))
+	{
+		BroadcastSelectionChanged();
+
+		++InsideSelChange;
+		BroadcastPreviewChanged();
+		--InsideSelChange;
+	}	
+}
+
+bool FPhysicsAssetEditorSharedData::IsCoMSelected(const int32 BodyIndex) const
+{
+	return SelectedCoMs.ContainsByPredicate([BodyIndex](const FSelection& Element) { return Element.Index == BodyIndex; });
+}
+
+void FPhysicsAssetEditorSharedData::ClearSelectedBody()
+{
+	ClearSelected();
+}
+
 void FPhysicsAssetEditorSharedData::SetSelectedBody(const FSelection& Body, bool bSelected)
 {
 	SetSelectedBodies({ Body }, bSelected);
@@ -739,37 +937,20 @@ void FPhysicsAssetEditorSharedData::SetSelectedBody(const FSelection& Body, bool
 
 void FPhysicsAssetEditorSharedData::SetSelectedBodies(const TArray<FSelection>& Bodies, bool bSelected)
 {
-	if (InsideSelChange || Bodies.Num() == 0)
+	if (!InsideSelChange && SetSelected(SelectedBodies, Bodies, bSelected))
 	{
-		return;
-	}
+		BroadcastSelectionChanged();
 
-	if (bSelected)
-	{
-		for (const FSelection& Body : Bodies)
+		if (!GetSelectedBody())
 		{
-			SelectedBodies.AddUnique(Body);
-		}
-	}
-	else
-	{
-		for (const FSelection& Body : Bodies)
-		{
-			SelectedBodies.Remove(Body);
+			return;
 		}
-	}
-
-	BroadcastSelectionChanged();
 
-	if (!GetSelectedBody())
-	{
-		return;
+		UpdateNoCollisionBodies();
+		++InsideSelChange;
+		BroadcastPreviewChanged();
+		--InsideSelChange;
 	}
-
-	UpdateNoCollisionBodies();
-	++InsideSelChange;
-	BroadcastPreviewChanged();
-	--InsideSelChange;
 }
 
 bool FPhysicsAssetEditorSharedData::IsBodySelected(const FSelection& Body) const
@@ -777,6 +958,11 @@ bool FPhysicsAssetEditorSharedData::IsBodySelected(const FSelection& Body) const
 	return SelectedBodies.Contains(Body);
 }
 
+bool FPhysicsAssetEditorSharedData::IsBodySelected(const int32 BodyIndex) const
+{
+	return SelectedBodies.ContainsByPredicate([BodyIndex](const FSelection& SelectedBody) { return SelectedBody.Index == BodyIndex; });
+}
+
 void FPhysicsAssetEditorSharedData::ToggleSelectionType(bool bIgnoreUserConstraints)
 {
 	TSet<int32> NewSelectedBodies; 
@@ -1186,10 +1372,7 @@ void FPhysicsAssetEditorSharedData::ClearSelectedConstraints()
 		return;
 	}
 
-	SelectedBodies.Empty();
-	SelectedConstraints.Empty();
-
-	BroadcastSelectionChanged();
+	ClearSelected();
 
 	++InsideSelChange;
 	BroadcastPreviewChanged();
@@ -1203,35 +1386,7 @@ void FPhysicsAssetEditorSharedData::SetSelectedConstraint(int32 ConstraintIndex,
 
 void FPhysicsAssetEditorSharedData::SetSelectedConstraints(const TArray<int32> ConstraintsIndices, bool bSelected)
 {
-	if (ConstraintsIndices.Num() == 0)
-	{
-		return;
-	}
-
-	if (InsideSelChange)
-	{
-		return;
-	}
-
-	bool bSelectionchanged = false;
-	for (int32 ConstraintIndex : ConstraintsIndices)
-	{
-		if (ConstraintIndex != INDEX_NONE)
-		{
-			FSelection Constraint(ConstraintIndex, EAggCollisionShape::Unknown, INDEX_NONE);
-			if (bSelected)
-			{
-				SelectedConstraints.AddUnique(Constraint);
-			}
-			else
-			{
-				SelectedConstraints.Remove(Constraint);
-			}
-			bSelectionchanged = true;
-		}
-	}
-
-	if (bSelectionchanged)
+	if (!InsideSelChange && SetSelected(SelectedConstraints, ConstraintsIndices, bSelected))
 	{
 		BroadcastSelectionChanged();
 
diff --git a/Engine/Source/Editor/PhysicsAssetEditor/Private/PhysicsAssetEditorSharedData.h b/Engine/Source/Editor/PhysicsAssetEditor/Private/PhysicsAssetEditorSharedData.h
--- a/Engine/Source/Editor/PhysicsAssetEditor/Private/PhysicsAssetEditorSharedData.h
+++ b/Engine/Source/Editor/PhysicsAssetEditor/Private/PhysicsAssetEditorSharedData.h
@@ -3,6 +3,7 @@
 #pragma once
 
 #include "CoreMinimal.h"
+#include "Math/Axis.h"
 #include "PreviewScene.h"
 #include "PhysicsAssetUtils.h"
 #include "PhysicsEngine/ConstraintInstance.h"
@@ -29,6 +30,25 @@ struct FScopedBulkSelection
 	TSharedPtr<FPhysicsAssetEditorSharedData> SharedData;
 };
 
+// BodyData
+//
+// A structure to keep track of editor only per physics body settings.
+//
+struct FBodyData
+{
+public:
+	FBodyData();
+
+	bool IsCoMAxisFixedInComponentSpace(const EAxis::Type InAxis) const;
+	void SetCoMAxisFixedInComponentSpace(const EAxis::Type InAxis, const bool bValue);
+
+private:
+	int32 GetBitFlag(const EAxis::Type InAxis) const;
+
+	int32 CoMAxisFixedInComponentSpaceFlags;
+};
+
+
 /*-----------------------------------------------------------------------------
    FPhysicsAssetEditorSharedData
 -----------------------------------------------------------------------------*/
@@ -55,10 +75,14 @@ public:
 		int32 PrimitiveIndex;
 		FTransform WidgetTM;
 		FTransform ManipulateTM;
-
-		FSelection(int32 GivenBodyIndex, EAggCollisionShape::Type GivenPrimitiveType, int32 GivenPrimitiveIndex) :
-			Index(GivenBodyIndex), PrimitiveType(GivenPrimitiveType), PrimitiveIndex(GivenPrimitiveIndex),
-			WidgetTM(FTransform::Identity), ManipulateTM(FTransform::Identity)
+		FVector CoMPosition; 
+
+		FSelection(int32 GivenBodyIndex, EAggCollisionShape::Type GivenPrimitiveType, int32 GivenPrimitiveIndex) 
+			: Index(GivenBodyIndex)
+			, PrimitiveType(GivenPrimitiveType)
+			, PrimitiveIndex(GivenPrimitiveIndex)
+			, WidgetTM(FTransform::Identity)
+			, ManipulateTM(FTransform::Identity)
 		{
 		}
 
@@ -77,6 +101,9 @@ public:
 	/** Accessor for mesh view mode, allows access for simulation and non-simulation modes */
 	EPhysicsAssetEditorMeshViewMode GetCurrentMeshViewMode(bool bSimulation);
 
+	/** Accessor for Center of Mass view mode, allows access for simulation and non-simulation modes */
+	EPhysicsAssetEditorCenterOfMassViewMode GetCurrentCenterOfMassViewMode(const bool bSimulation) const;
+
 	/** Accessor for collision view mode, allows access for simulation and non-simulation modes */
 	EPhysicsAssetEditorCollisionViewMode GetCurrentCollisionViewMode(bool bSimulation);
 
@@ -139,11 +166,18 @@ public:
 	/** Gets whether the active constraint is locked */
 	bool IsAngularConstraintLocked(EPhysicsAssetEditorConstraintType Constraint) const;
 
+	/** Center of Mass editing */
+	void ClearSelectedCoMs();
+	void SetSelectedCoM(const FSelection& InSelectedCoM, const bool bSelected);
+	void SetSelectedCoMs(const TArray<FSelection>& InSelectedElements, const bool bSelected);
+	bool IsCoMSelected(const int32 BodyIndex) const;
+
 	/** Collision geometry editing */
 	void ClearSelectedBody();
 	void SetSelectedBody(const FSelection& Body, bool bSelected);
 	void SetSelectedBodies(const TArray<FSelection>& Bodies, bool bSelected);
 	bool IsBodySelected(const FSelection& Body) const;
+	bool IsBodySelected(const int32 BodyIndex) const;
 	void ToggleSelectionType(bool bIgnoreUserConstraints = true);
 	void ToggleShowSelected();
 	bool IsBodyHidden(const int32 BodyIndex) const;
@@ -204,6 +238,9 @@ public:
 	/** Handle clicking on a body */
 	void HitBone(int32 BodyIndex, EAggCollisionShape::Type PrimType, int32 PrimIndex, bool bGroupSelect);
 
+	/** Handle clicking on a Center of Mass marker */
+	void HitCoM(const int32 BodyIndex, const bool bGroupSelect);
+
 	/** Handle clikcing on a constraint */
 	void HitConstraint(int32 ConstraintIndex, bool bGroupSelect);
 
@@ -254,6 +291,12 @@ public:
 	void SetShowCom(bool InValue);
 	bool GetShowCom() const;
 
+	/** Returns the correct location to draw a CoM marker in the viewport */
+	FVector GetCOMRenderPosition(const int32 BodyIndex) const;
+
+	bool IsCoMAxisFixedInComponentSpace(const int32 BodyIndex, const EAxis::Type InAxis) const;
+	void SetCoMAxisFixedInComponentSpace(const int32 BodyIndex, const EAxis::Type InAxis, const bool bValue);
+
 private:
 	/** Initializes a constraint setup */
 	void InitConstraintSetup(UPhysicsConstraintTemplate* ConstraintSetup, int32 ChildBodyIndex, int32 ParentBodyIndex);
@@ -279,12 +322,15 @@ private:
 	/** Checks and parses clipboard data */
 	static bool ParseClipboard(UPhysicsAsset*& OutAsset, FString& OutObjectType, UObject*& OutObject);
 
-	/** Gneerate a new unique name for a constraint */
+	/** Generate a new unique name for a constraint */
 	FString MakeUniqueNewConstraintName();
 
+	/** Clears all the selected objects */
+	void ClearSelected();
+
 public:
 	/** Callback for handling selection changes */
-	DECLARE_EVENT_TwoParams(FPhysicsAssetEditorSharedData, FSelectionChanged, const TArray<FSelection>&, const TArray<FSelection>&);
+	DECLARE_EVENT_ThreeParams(FPhysicsAssetEditorSharedData, FSelectionChanged, const TArray<FSelection>&, const TArray<FSelection>&, const TArray<FSelection>&);
 	FSelectionChanged SelectionChangedEvent;
 
 	/** Callback for handling changes to the bone/body/constraint hierarchy */
@@ -327,10 +373,16 @@ public:
 	/** List of bodies that don't collide with the currently selected collision body */
 	TArray<int32> NoCollisionBodies;
 
+	/** Center of Mass editing */
+	TArray<FSelection> SelectedCoMs;
+
 	/** Bone info */
 	TArray<FBoneVertInfo> DominantWeightBoneInfos;
 	TArray<FBoneVertInfo> AnyWeightBoneInfos;
 
+	/** Editor only data for each physics body. */
+	TArray<FBodyData> EditorBodyData;
+
 	TArray<FSelection> SelectedBodies;
 
 	FSelection * GetSelectedBody()
@@ -339,6 +391,12 @@ public:
 		return Count ? &SelectedBodies[Count - 1] : NULL;
 	}
 
+	FSelection* GetSelectedCoM()
+	{
+		int32 Count = SelectedCoMs.Num();
+		return Count ? &SelectedCoMs[Count - 1] : NULL;
+	}
+
 	/** Constraint editing */
 	TArray<FSelection> SelectedConstraints;
 	FSelection * GetSelectedConstraint()
@@ -362,6 +420,8 @@ public:
 	/** Manipulation (rotate, translate, scale) */
 	bool bManipulating;
 
+	bool bShouldUpdatedSelectedCoMs = false;
+
 	/** when true, we dont broadcast every selection change - allows for bulk changes without so much overhead */
 	bool bSuspendSelectionBroadcast;
 
diff --git a/Engine/Source/Editor/PhysicsAssetEditor/Private/PhysicsAssetEditorSkeletalMeshComponent.cpp b/Engine/Source/Editor/PhysicsAssetEditor/Private/PhysicsAssetEditorSkeletalMeshComponent.cpp
--- a/Engine/Source/Editor/PhysicsAssetEditor/Private/PhysicsAssetEditorSkeletalMeshComponent.cpp
+++ b/Engine/Source/Editor/PhysicsAssetEditor/Private/PhysicsAssetEditorSkeletalMeshComponent.cpp
@@ -112,20 +112,20 @@ void UPhysicsAssetEditorSkeletalMeshComponent::DebugDraw(const FSceneView* View,
 
 	ElemSelectedMaterial->SetVectorParameterValue(SelectionColorName, LinearSelectionColor);
 
-	
-
 	FPhysicsAssetRenderSettings* const RenderSettings = UPhysicsAssetRenderUtilities::GetSettings(PhysicsAsset);
 	
 	if (RenderSettings)
 	{
 		// Copy render settings from editor viewport. These settings must be applied to the rendering in all editors 
 		// when an asset is open in the Physics Asset Editor but should not persist after the editor has been closed.
+		RenderSettings->CenterOfMassViewMode = SharedData->GetCurrentCenterOfMassViewMode(SharedData->bRunningSimulation);
 		RenderSettings->CollisionViewMode = SharedData->GetCurrentCollisionViewMode(SharedData->bRunningSimulation);
 		RenderSettings->ConstraintViewMode = SharedData->GetCurrentConstraintViewMode(SharedData->bRunningSimulation);
 		RenderSettings->ConstraintDrawSize = SharedData->EditorOptions->ConstraintDrawSize;
 		RenderSettings->PhysicsBlend = SharedData->EditorOptions->PhysicsBlend;
 		RenderSettings->bHideKinematicBodies = SharedData->EditorOptions->bHideKinematicBodies;
 		RenderSettings->bHideSimulatedBodies = SharedData->EditorOptions->bHideSimulatedBodies;
+		RenderSettings->bHideBodyMass = SharedData->EditorOptions->bHideBodyMass;
 		RenderSettings->bRenderOnlySelectedConstraints = SharedData->EditorOptions->bRenderOnlySelectedConstraints;
 		RenderSettings->bShowConstraintsAsPoints = SharedData->EditorOptions->bShowConstraintsAsPoints;
 		RenderSettings->bDrawViolatedLimits = SharedData->EditorOptions->bDrawViolatedLimits;
@@ -139,6 +139,13 @@ void UPhysicsAssetEditorSkeletalMeshComponent::DebugDraw(const FSceneView* View,
 
 			PhysicsAssetRender::DebugDrawBodies(this, PhysicsAsset, PDI, ColorFn, MaterialFn, TransformFn, HitProxyFn);
 		}
+		
+		{
+			auto COMPositionFn = [this](const int32 BodyIndex) { return this->SharedData->GetCOMRenderPosition(BodyIndex);  };
+			auto IsSelectedFn = [this](const uint32 InIndex) { return this->SharedData->IsBodySelected(InIndex) || this->SharedData->IsCoMSelected(InIndex); };
+			auto HitProxyFn = [](const int32 BodyIndex) { return new HPhysicsAssetEditorEdCoMProxy(BodyIndex); };
+			PhysicsAssetRender::DebugDrawCenterOfMass(this, PhysicsAsset, PDI, COMPositionFn, IsSelectedFn, HitProxyFn);
+		}
 
 		// Draw Constraints.
 		{
diff --git a/Engine/Source/Editor/UnrealEd/Classes/Preferences/PhysicsAssetEditorOptions.h b/Engine/Source/Editor/UnrealEd/Classes/Preferences/PhysicsAssetEditorOptions.h
--- a/Engine/Source/Editor/UnrealEd/Classes/Preferences/PhysicsAssetEditorOptions.h
+++ b/Engine/Source/Editor/UnrealEd/Classes/Preferences/PhysicsAssetEditorOptions.h
@@ -18,6 +18,14 @@ enum class EPhysicsAssetEditorCollisionViewMode : uint8
 	None
 };
 
+UENUM()
+enum class EPhysicsAssetEditorCenterOfMassViewMode : uint8
+{
+	All,
+	Selected,
+	None
+};
+
 UENUM()
 enum class EPhysicsAssetEditorMeshViewMode : uint8
 {
@@ -125,12 +133,16 @@ class UPhysicsAssetEditorOptions : public UObject
 
 	/** View mode for meshes in edit mode */
 	UPROPERTY(config)
-		EPhysicsAssetEditorMeshViewMode MeshViewMode;
+	EPhysicsAssetEditorMeshViewMode MeshViewMode;
 
 	/** View mode for collision in edit mode */
 	UPROPERTY(config)
 	EPhysicsAssetEditorCollisionViewMode CollisionViewMode;
 
+	/** View mode for Center of Mass in edit mode */
+	UPROPERTY(config)
+	EPhysicsAssetEditorCenterOfMassViewMode CenterOfMassViewMode;
+
 	/** View mode for constraints in edit mode */
 	UPROPERTY(config)
 	EPhysicsAssetEditorConstraintViewMode ConstraintViewMode;
@@ -143,6 +155,10 @@ class UPhysicsAssetEditorOptions : public UObject
 	UPROPERTY(config)
 	EPhysicsAssetEditorCollisionViewMode SimulationCollisionViewMode;
 
+	/** View mode for Center of Mass in simulation mode */
+	UPROPERTY(config)
+	EPhysicsAssetEditorCenterOfMassViewMode SimulationCenterOfMassViewMode;
+	
 	/** View mode for constraints in simulation mode */
 	UPROPERTY(config)
 	EPhysicsAssetEditorConstraintViewMode SimulationConstraintViewMode;
@@ -163,6 +179,10 @@ class UPhysicsAssetEditorOptions : public UObject
 	UPROPERTY(config)
 	bool bHideKinematicBodies;
 
+	/** When set, disables rendering for body mass text displayed next to Center of Mass indicator */
+	UPROPERTY(config)
+	bool bHideBodyMass;
+
 	/** When set, cloth will reset each time simulation is toggled */
 	UPROPERTY(EditAnywhere, config, Category=Clothing)
 	bool bResetClothWhenSimulating;
diff --git a/Engine/Source/Runtime/Engine/Classes/PhysicsEngine/BodyInstance.h b/Engine/Source/Runtime/Engine/Classes/PhysicsEngine/BodyInstance.h
--- a/Engine/Source/Runtime/Engine/Classes/PhysicsEngine/BodyInstance.h
+++ b/Engine/Source/Runtime/Engine/Classes/PhysicsEngine/BodyInstance.h
@@ -440,7 +440,7 @@ protected:
 	/** 
 	 * @brief Enable automatic inertia conditioning to stabilize constraints.
 	 * 
-	 * Inertia conitioning increases inertia when an object is long and thin and also when it has joints that are outside the
+	 * Inertia conditioning increases inertia when an object is long and thin and also when it has joints that are outside the
 	 * collision shapes of the body. Increasing the inertia reduces the amount of rotation applied at joints which helps stabilize
 	 * joint chains, especially when bodies are small. In principle you can get the same behaviour by setting the InertiaTensorScale
 	 * appropriately, but this takes some of the guesswork out of it.
@@ -556,7 +556,7 @@ public:
 	UPROPERTY(EditAnywhere, Category = Physics, meta = (DisplayName = "Plane Normal"))
 	FVector CustomDOFPlaneNormal;
 
-	/** User specified offset for the center of mass of this object, from the calculated location */
+	/** User specified offset for this object's Center of Mass. The offset is defined in bone space and will be added to the calculated location. */
 	UPROPERTY(EditAnywhere, AdvancedDisplay, BlueprintReadWrite, Category = Physics, meta = (DisplayName = "Center Of Mass Offset"))
 	FVector COMNudge;
 
