[UE] Fix for static analysis warning
- Issue 462739: Static analyzer warnings in PhysicsAssetEditorEditMode.cpp

#rnx

diff --git a/Engine/Source/Editor/PhysicsAssetEditor/Private/PhysicsAssetEditorEditMode.cpp b/Engine/Source/Editor/PhysicsAssetEditor/Private/PhysicsAssetEditorEditMode.cpp
--- a/Engine/Source/Editor/PhysicsAssetEditor/Private/PhysicsAssetEditorEditMode.cpp
+++ b/Engine/Source/Editor/PhysicsAssetEditor/Private/PhysicsAssetEditorEditMode.cpp
@@ -1271,7 +1271,7 @@ void FPhysicsAssetEditorEditMode::UpdateCoM()
 
 FVector FPhysicsAssetEditorEditMode::CalculateCoMNudgeForWorldSpacePosition(const int32 BodyIndex, const FVector& CoMPositionWorldSpace) const
 {
-	FVector	CalculatedCoMOffset;
+	FVector	CalculatedCoMOffset = FVector::ZeroVector;
 
 	if (SharedData->EditorSkelComp->Bodies.IsValidIndex(BodyIndex))
 	{
