[UE] Center of Mass editing tools in PhAT - remove conflicting Mass Properties check box
- removed Mass Properties check box (as it overlapped with new CoM / mass drawing options)
- renamed Center of Mass sub menu to Mass Properties to make it easier for people to find the old functionality in the new settings
- removed ToggleMassProperties command action
- ensured CoM drawing is disabled by default

diff --git a/Engine/Source/Editor/PhysicsAssetEditor/Private/PhysicsAssetEditor.cpp b/Engine/Source/Editor/PhysicsAssetEditor/Private/PhysicsAssetEditor.cpp
--- a/Engine/Source/Editor/PhysicsAssetEditor/Private/PhysicsAssetEditor.cpp
+++ b/Engine/Source/Editor/PhysicsAssetEditor/Private/PhysicsAssetEditor.cpp
@@ -775,7 +775,32 @@ void FPhysicsAssetEditor::ExtendViewportMenus()
 		if (PhysicsAssetEditor)
 		{
 			FToolMenuSection& Section = CharacterMenu->AddSection("PhysicsAssetShowCommands", LOCTEXT("PhysicsShowCommands", "Physics Rendering"), FToolMenuInsert("AnimViewportSceneElements", EToolMenuInsertType::Before));
-			Section.AddMenuEntry(FPhysicsAssetEditorCommands::Get().ToggleMassProperties);
+
+			Section.AddSubMenu(TEXT("MassPropertiesSubMenu"), LOCTEXT("MassPropertiesSubMenu", "Mass Properties"), FText::GetEmpty(),
+				FNewToolMenuDelegate::CreateLambda([WeakPhysicsAssetEditor = PhysicsAssetEditor.ToWeakPtr()](UToolMenu* InSubMenu)
+					{
+						const FPhysicsAssetEditorCommands& Commands = FPhysicsAssetEditorCommands::Get();
+
+						{
+							FToolMenuSection& Section = InSubMenu->AddSection("PhysicsAssetEditorCenterOfMassRenderSettings", LOCTEXT("CenterOfMassRenderSettingsHeader", "Center of Mass Drawing"));
+							Section.AddMenuEntry(Commands.HideBodyMass);
+						}
+
+						{
+							FToolMenuSection& Section = InSubMenu->AddSection("PhysicsAssetEditorCenterOfMassRenderingMode", LOCTEXT("CenterOfMassRenderingModeHeader", "Center of Mass Drawing (Edit)"));
+							Section.AddMenuEntry(Commands.CenterOfMassRenderingMode_All);
+							Section.AddMenuEntry(Commands.CenterOfMassRenderingMode_Selected);
+							Section.AddMenuEntry(Commands.CenterOfMassRenderingMode_None);
+						}
+
+						{
+							FToolMenuSection& Section = InSubMenu->AddSection("PhysicsAssetEditorCenterOfMassRenderingModeSim", LOCTEXT("CenterOfMassRenderingModeSimHeader", "Center of Mass Drawing (Simulation)"));
+							Section.AddMenuEntry(Commands.CenterOfMassRenderingMode_Simulation_All);
+							Section.AddMenuEntry(Commands.CenterOfMassRenderingMode_Simulation_Selected);
+							Section.AddMenuEntry(Commands.CenterOfMassRenderingMode_Simulation_None);
+						}
+					}));
+
 			Section.AddSubMenu(TEXT("MeshRenderModeSubMenu"), LOCTEXT("MeshRenderModeSubMenu", "Mesh"), FText::GetEmpty(),
 				FNewToolMenuDelegate::CreateLambda([](UToolMenu* InSubMenu)
 			{
@@ -824,31 +849,6 @@ void FPhysicsAssetEditor::ExtendViewportMenus()
 				}
 			}));
 
-			Section.AddSubMenu(TEXT("CenterOfMassModeSubMenu"), LOCTEXT("CenterOfMassModeSubMenu", "Center of Mass"), FText::GetEmpty(),
-				FNewToolMenuDelegate::CreateLambda([WeakPhysicsAssetEditor = PhysicsAssetEditor.ToWeakPtr()](UToolMenu* InSubMenu)
-			{
-				const FPhysicsAssetEditorCommands& Commands = FPhysicsAssetEditorCommands::Get();
-
-				{
-					FToolMenuSection& Section = InSubMenu->AddSection("PhysicsAssetEditorCenterOfMassRenderSettings", LOCTEXT("CenterOfMassRenderSettingsHeader", "Center of Mass Drawing"));
-					Section.AddMenuEntry(Commands.HideBodyMass);
-				}
-
-				{
-					FToolMenuSection& Section = InSubMenu->AddSection("PhysicsAssetEditorCenterOfMassRenderingMode", LOCTEXT("CenterOfMassRenderingModeHeader", "Center of Mass Drawing (Edit)"));
-					Section.AddMenuEntry(Commands.CenterOfMassRenderingMode_All);
-					Section.AddMenuEntry(Commands.CenterOfMassRenderingMode_Selected);
-					Section.AddMenuEntry(Commands.CenterOfMassRenderingMode_None);
-				}
-
-				{
-					FToolMenuSection& Section = InSubMenu->AddSection("PhysicsAssetEditorCenterOfMassRenderingModeSim", LOCTEXT("CenterOfMassRenderingModeSimHeader", "Center of Mass Drawing (Simulation)"));
-					Section.AddMenuEntry(Commands.CenterOfMassRenderingMode_Simulation_All);
-					Section.AddMenuEntry(Commands.CenterOfMassRenderingMode_Simulation_Selected);
-					Section.AddMenuEntry(Commands.CenterOfMassRenderingMode_Simulation_None);
-				}
-			}));
-
 			Section.AddSubMenu(TEXT("ConstraintConstraintModeSubMenu"), LOCTEXT("ConstraintConstraintModeSubMenu", "Constraints"), FText::GetEmpty(),
 				FNewToolMenuDelegate::CreateLambda([WeakPhysicsAssetEditor = PhysicsAssetEditor.ToWeakPtr()](UToolMenu* InSubMenu)
 			{
@@ -1546,12 +1546,6 @@ void FPhysicsAssetEditor::BindCommands()
 		FCanExecuteAction(),
 		FIsActionChecked::CreateSP(this, &FPhysicsAssetEditor::IsRenderingOnlySelectedConstraints));
 
-	ViewportCommandList->MapAction(
-		Commands.ToggleMassProperties,
-		FExecuteAction::CreateSP(this, &FPhysicsAssetEditor::OnToggleMassProperties),
-		FCanExecuteAction(),
-		FIsActionChecked::CreateSP(this, &FPhysicsAssetEditor::IsToggleMassProperties));
-
 	SkeletonTreeCommandList = MakeShared<FUICommandList_Pinnable>();
 
 	SkeletonTreeCommandList->MapAction(
diff --git a/Engine/Source/Editor/PhysicsAssetEditor/Private/PhysicsAssetEditorActions.cpp b/Engine/Source/Editor/PhysicsAssetEditor/Private/PhysicsAssetEditorActions.cpp
--- a/Engine/Source/Editor/PhysicsAssetEditor/Private/PhysicsAssetEditorActions.cpp
+++ b/Engine/Source/Editor/PhysicsAssetEditor/Private/PhysicsAssetEditorActions.cpp
@@ -62,7 +62,6 @@ void FPhysicsAssetEditorCommands::RegisterCommands()
 	UI_COMMAND(DrawConstraintsAsPoints, "Draw Constraints As Points", "Draw Constraints As Points", EUserInterfaceActionType::ToggleButton, FInputChord());
 	UI_COMMAND(DrawViolatedLimits, "Draw Violated Limits", "Draw Violated Limits", EUserInterfaceActionType::ToggleButton, FInputChord());
 	UI_COMMAND(HideBodyMass, "Hide Mass", "Hide text displaying body mass next to CoM", EUserInterfaceActionType::ToggleButton, FInputChord());
-	UI_COMMAND(ToggleMassProperties, "Mass Properties", "Show Mass Properties For Bodies When Simulation Is Enabled", EUserInterfaceActionType::ToggleButton, FInputChord());
 	UI_COMMAND(DisableCollision, "Disable Collision", "Disable collision between the currently selected bodies", EUserInterfaceActionType::Button, FInputChord(EKeys::RightBracket));
 	UI_COMMAND(DisableCollisionAll, "Disable Collision All", "Disable collision between the currently selected bodies and all bodies", EUserInterfaceActionType::Button, FInputChord(EModifierKey::Control, EKeys::RightBracket));
 	UI_COMMAND(EnableCollision, "Enable Collision", "Enable collision between the currently selected bodies", EUserInterfaceActionType::Button, FInputChord(EKeys::LeftBracket));
diff --git a/Engine/Source/Editor/PhysicsAssetEditor/Private/PhysicsAssetEditorActions.h b/Engine/Source/Editor/PhysicsAssetEditor/Private/PhysicsAssetEditorActions.h
--- a/Engine/Source/Editor/PhysicsAssetEditor/Private/PhysicsAssetEditorActions.h
+++ b/Engine/Source/Editor/PhysicsAssetEditor/Private/PhysicsAssetEditorActions.h
@@ -76,7 +76,6 @@ public:
 	TSharedPtr<FUICommandInfo> RenderOnlySelectedConstraints;
 	TSharedPtr<FUICommandInfo> DrawConstraintsAsPoints;
 	TSharedPtr<FUICommandInfo> DrawViolatedLimits;
-	TSharedPtr<FUICommandInfo> ToggleMassProperties;
 	TSharedPtr<FUICommandInfo> DisableCollision;
 	TSharedPtr<FUICommandInfo> DisableCollisionAll;
 	TSharedPtr<FUICommandInfo> EnableCollision;
diff --git a/Engine/Source/Editor/UnrealEd/Private/PreferenceStubs.cpp b/Engine/Source/Editor/UnrealEd/Private/PreferenceStubs.cpp
--- a/Engine/Source/Editor/UnrealEd/Private/PreferenceStubs.cpp
+++ b/Engine/Source/Editor/UnrealEd/Private/PreferenceStubs.cpp
@@ -49,8 +49,10 @@ UPhysicsAssetEditorOptions::UPhysicsAssetEditorOptions(const FObjectInitializer&
 
 	// view options
 	MeshViewMode = EPhysicsAssetEditorMeshViewMode::Solid;
+	CenterOfMassViewMode = EPhysicsAssetEditorCenterOfMassViewMode::None;
 	CollisionViewMode = EPhysicsAssetEditorCollisionViewMode::Solid;
 	ConstraintViewMode = EPhysicsAssetEditorConstraintViewMode::AllLimits;
+	SimulationCenterOfMassViewMode = EPhysicsAssetEditorCenterOfMassViewMode::None;
 	SimulationMeshViewMode = EPhysicsAssetEditorMeshViewMode::Solid;
 	SimulationCollisionViewMode = EPhysicsAssetEditorCollisionViewMode::Solid;
 	SimulationConstraintViewMode = EPhysicsAssetEditorConstraintViewMode::None;
@@ -59,6 +61,7 @@ UPhysicsAssetEditorOptions::UPhysicsAssetEditorOptions(const FObjectInitializer&
 	bSolidRenderingForSelectedOnly = false;
 	bHideSimulatedBodies = false;
 	bHideKinematicBodies = false;
+	bHideBodyMass = true;
 	bResetClothWhenSimulating = false;
 }
 
