[UE] Changes to PhAT CoM editing tools following user feedback.
- Change hide mass value check box polarity / text in menu options
- Make CoM markers scale option
- Make details panel component shorter so its easier to see all the components
- Align yellow CoM markers with bone space
- Stop Body details panel hiding when the CoM marker is selected

#rnx

diff --git a/Engine/Source/Editor/Persona/Private/PhysicsAssetRenderUtils.cpp b/Engine/Source/Editor/Persona/Private/PhysicsAssetRenderUtils.cpp
--- a/Engine/Source/Editor/Persona/Private/PhysicsAssetRenderUtils.cpp
+++ b/Engine/Source/Editor/Persona/Private/PhysicsAssetRenderUtils.cpp
@@ -41,11 +41,16 @@ FTransform GetConstraintMatrix(const USkeletalMeshComponent* const SkeletalMeshC
 	return LFrame * BoneTM;
 }
 
-void DrawWireStar(FPrimitiveDrawInterface* PDI, const FVector& Position, float Size, const FLinearColor& Color, uint8 DepthPriority, const float Thickness)
+void DrawWireStar(FPrimitiveDrawInterface* PDI, const FTransform& Transform, float Size, const FLinearColor& Color, uint8 DepthPriority, const float Thickness)
 {
-	PDI->DrawLine(Position + Size * FVector(1, 0, 0), Position - Size * FVector(1, 0, 0), Color, DepthPriority, Thickness);
-	PDI->DrawLine(Position + Size * FVector(0, 1, 0), Position - Size * FVector(0, 1, 0), Color, DepthPriority, Thickness);
-	PDI->DrawLine(Position + Size * FVector(0, 0, 1), Position - Size * FVector(0, 0, 1), Color, DepthPriority, Thickness);
+	const FVector Position = Transform.GetLocation();
+	const FVector XAxis = Transform.GetUnitAxis(EAxis::X);
+	const FVector YAxis = Transform.GetUnitAxis(EAxis::Y);
+	const FVector ZAxis = Transform.GetUnitAxis(EAxis::Z);
+
+	PDI->DrawLine(Position + Size * XAxis, Position - Size * XAxis, Color, DepthPriority, Thickness);
+	PDI->DrawLine(Position + Size * YAxis, Position - Size * YAxis, Color, DepthPriority, Thickness);
+	PDI->DrawLine(Position + Size * ZAxis, Position - Size * ZAxis, Color, DepthPriority, Thickness);
 }
 
 ////////////////////////////////////////
@@ -69,7 +74,7 @@ FPhysicsAssetRenderSettings::FPhysicsAssetRenderSettings()
 	, BoneUnselectedColor(170, 155, 225)
 	, NoCollisionColor(200, 200, 200)
 	, COMRenderColor(255, 255, 100)
-	, COMRenderSize(3.0f)
+	, COMRenderSize(2.0f)
 	, COMRenderLineThickness(0.2f)
 	, COMRenderMassTextOffsetScreenspace(8.0f)
 	, InfluenceLineLength(2.0f)
@@ -267,6 +272,7 @@ void FPhysicsAssetRenderSettings::ResetEditorViewportOptions()
 
 	CenterOfMassViewMode = DefaultObject.CenterOfMassViewMode;
 	CollisionViewMode = DefaultObject.CollisionViewMode;
+	COMRenderSize = DefaultObject.COMRenderSize;
 	ConstraintViewMode = DefaultObject.ConstraintViewMode;
 	ConstraintDrawSize = DefaultObject.ConstraintDrawSize;
 	PhysicsBlend = DefaultObject.PhysicsBlend;
@@ -474,8 +480,16 @@ namespace PhysicsAssetRender
 							COMRenderLineThickness *= SelectedItemRenderSizeMultiplier;
 						}
 
+						FTransform CoMMarkerTM(GetCoMPosition(BodyIndex));
+						
+						if (BodyInstance->GetBodySetup())
+						{
+							const FName BoneName = BodyInstance->GetBodySetup()->BoneName;
+							CoMMarkerTM.SetRotation(SkeletalMeshComponent->GetBoneTransform(BoneName).GetRotation());
+						}
+
 						PDI->SetHitProxy(CreateHitProxy(BodyIndex));
-						DrawWireStar(PDI, GetCoMPosition(BodyIndex), COMRenderSize, RenderSettings->COMRenderColor, SDPG_Foreground, COMRenderLineThickness);
+						DrawWireStar(PDI, CoMMarkerTM, COMRenderSize, RenderSettings->COMRenderColor, SDPG_Foreground, COMRenderLineThickness);
 						PDI->SetHitProxy(NULL);
 					}
 				}
diff --git a/Engine/Source/Editor/PhysicsAssetEditor/Private/PhysicsAssetEditor.cpp b/Engine/Source/Editor/PhysicsAssetEditor/Private/PhysicsAssetEditor.cpp
--- a/Engine/Source/Editor/PhysicsAssetEditor/Private/PhysicsAssetEditor.cpp
+++ b/Engine/Source/Editor/PhysicsAssetEditor/Private/PhysicsAssetEditor.cpp
@@ -309,6 +309,10 @@ void FPhysicsAssetEditor::HandleViewportSelectionChanged(const TArray<FPhysicsAs
 			{ 
 				return SharedData->PhysicsAsset->ConstraintSetup[InItem.Index];
 			});
+			Algo::Transform(InSelectedCoMs, Objects, [this](const FPhysicsAssetEditorSharedData::FSelection& InItem)
+			{
+					return SharedData->PhysicsAsset->SkeletalBodySetups[InItem.Index]; // Add the owning physics body here so we display its details panel in the UI.
+			});
 			Algo::Transform(InSelectedBodies, Bodies, [this](const FPhysicsAssetEditorSharedData::FSelection& InItem) 
 			{ 
 				return SharedData->PhysicsAsset->SkeletalBodySetups[InItem.Index];
@@ -783,7 +787,8 @@ void FPhysicsAssetEditor::ExtendViewportMenus()
 
 						{
 							FToolMenuSection& Section = InSubMenu->AddSection("PhysicsAssetEditorCenterOfMassRenderSettings", LOCTEXT("CenterOfMassRenderSettingsHeader", "Center of Mass Drawing"));
-							Section.AddMenuEntry(Commands.HideBodyMass);
+							Section.AddMenuEntry(Commands.DrawBodyMass);
+							Section.AddEntry(FToolMenuEntry::InitWidget(TEXT("CoMMarkerScale"), WeakPhysicsAssetEditor.Pin()->MakeCoMMarkerScaleWidget(), LOCTEXT("CoMMarkerScaleLabel", "Marker Scale")));
 						}
 
 						{
@@ -1523,10 +1528,10 @@ void FPhysicsAssetEditor::BindCommands()
 		FIsActionChecked::CreateSP(this, &FPhysicsAssetEditor::IsHidingKinematicBodies));
 	
 	ViewportCommandList->MapAction(
-		Commands.HideBodyMass,
+		Commands.DrawBodyMass,
 		FExecuteAction::CreateSP(this, &FPhysicsAssetEditor::ToggleHideBodyMass),
 		FCanExecuteAction(),
-		FIsActionChecked::CreateSP(this, &FPhysicsAssetEditor::IsHidingBodyMass));
+		FIsActionChecked::CreateSP(this, &FPhysicsAssetEditor::IsDrawingBodyMass));
 
 	ViewportCommandList->MapAction(
 		Commands.DrawConstraintsAsPoints,
@@ -2911,6 +2916,11 @@ bool FPhysicsAssetEditor::IsHidingBodyMass() const
 	return SharedData->EditorOptions->bHideBodyMass;
 }
 
+bool FPhysicsAssetEditor::IsDrawingBodyMass() const
+{
+	return !IsHidingBodyMass();
+}
+
 bool FPhysicsAssetEditor::IsConstraintRenderingMode(EPhysicsAssetEditorConstraintViewMode Mode, bool bSimulation) const
 {
 	return Mode == SharedData->GetCurrentConstraintViewMode(bSimulation);
@@ -4024,33 +4034,43 @@ void FPhysicsAssetEditor::RecreatePhysicsState()
 	SharedData->EnableSimulation(false);
 }
 
-TSharedRef<SWidget> FPhysicsAssetEditor::MakeConstraintScaleWidget()
+template<typename TValueAccessor> TSharedRef<SWidget> FPhysicsAssetEditor::MakeScaleWidget(const float MinValue, const float MaxValue, TValueAccessor ValueAccessorFunction, const FName WidgetInteractionText)
 {
-	return 
+	return
 		SNew(SBox)
 		.HAlign(HAlign_Right)
 		[
 			SNew(SBox)
-			.Padding(FMargin(4.0f, 0.0f, 0.0f, 0.0f))
-			.WidthOverride(100.0f)
-			[
-				SNew(SNumericEntryBox<float>)
-				.Font(FAppStyle::GetFontStyle(TEXT("MenuItem.Font")))
-				.AllowSpin(true)
-				.MinSliderValue(0.0f)
-				.MaxSliderValue(4.0f)
-				.Value_Lambda([this]() { return SharedData->EditorOptions->ConstraintDrawSize; })
-				.OnValueChanged_Lambda([this](float InValue) { SharedData->EditorOptions->ConstraintDrawSize = InValue; })
-				.OnValueCommitted_Lambda([this](float InValue, ETextCommit::Type InCommitType) 
-				{
-					SharedData->EditorOptions->ConstraintDrawSize = InValue; 
-					SharedData->EditorOptions->SaveConfig(); 
-					ViewportCommandList->WidgetInteraction(TEXT("ConstraintScaleWidget"));
-				})
-			]
+				.Padding(FMargin(4.0f, 0.0f, 0.0f, 0.0f))
+				.WidthOverride(100.0f)
+				[
+					SNew(SNumericEntryBox<float>)
+						.Font(FAppStyle::GetFontStyle(TEXT("MenuItem.Font")))
+						.AllowSpin(true)
+						.MinSliderValue(MinValue)
+						.MaxSliderValue(MaxValue)
+						.Value_Lambda([ValueAccessorFunction]() { return ValueAccessorFunction(); })
+						.OnValueChanged_Lambda([ValueAccessorFunction](float InValue) { ValueAccessorFunction() = InValue; })
+						.OnValueCommitted_Lambda([this, ValueAccessorFunction, WidgetInteractionText](float InValue, ETextCommit::Type InCommitType)
+							{
+								ValueAccessorFunction() = InValue;
+								SharedData->EditorOptions->SaveConfig();
+								ViewportCommandList->WidgetInteraction(WidgetInteractionText);
+							})
+				]
 		];
 }
 
+TSharedRef<SWidget> FPhysicsAssetEditor::MakeConstraintScaleWidget()
+{
+	return MakeScaleWidget(0.0f, 4.0f, [this]() -> float& { return SharedData->EditorOptions->ConstraintDrawSize; }, TEXT("ConstraintScaleWidget"));
+}
+
+TSharedRef<SWidget> FPhysicsAssetEditor::MakeCoMMarkerScaleWidget()
+{
+	return MakeScaleWidget(0.0f, 4.0f, [this]() -> float& { return SharedData->EditorOptions->COMRenderSize; }, TEXT("CoMMarkerScaleWidget"));
+}
+
 TSharedRef<SWidget> FPhysicsAssetEditor::MakeCollisionOpacityWidget()
 {
 	return 
diff --git a/Engine/Source/Editor/PhysicsAssetEditor/Private/PhysicsAssetEditor.h b/Engine/Source/Editor/PhysicsAssetEditor/Private/PhysicsAssetEditor.h
--- a/Engine/Source/Editor/PhysicsAssetEditor/Private/PhysicsAssetEditor.h
+++ b/Engine/Source/Editor/PhysicsAssetEditor/Private/PhysicsAssetEditor.h
@@ -146,6 +146,9 @@ public:
 	/** Make the constraint scale widget */
 	TSharedRef<SWidget> MakeConstraintScaleWidget();
 
+	/** Make the Center of Mass marker scale widget */
+	TSharedRef<SWidget> MakeCoMMarkerScaleWidget();
+
 	/** Make the collision opacity widget */
 	TSharedRef<SWidget> MakeCollisionOpacityWidget();
 
@@ -258,6 +261,7 @@ private:
 	bool IsHidingSimulatedBodies() const;
 	bool IsHidingKinematicBodies() const;
 	bool IsHidingBodyMass() const;
+	bool IsDrawingBodyMass() const;
 	void OnToggleMassProperties();
 	bool IsToggleMassProperties() const;
 	void OnSetCollision(bool bEnable);
@@ -365,6 +369,9 @@ private:
 	/** show a notification message **/
 	void ShowNotificationMessage(const FText& Message, const SNotificationItem::ECompletionState CompletionState);
 
+	/** Helper function for creating Scale Widgets */
+	template<typename TValueAccessor> TSharedRef<SWidget> MakeScaleWidget(const float MinValue, const float MaxValue, TValueAccessor ValueAccessorFunction, const FName WidgetInteractionText);
+
 private:
 	/** Physics asset properties tab */
 	TSharedPtr<class IDetailsView> PhysAssetProperties;
diff --git a/Engine/Source/Editor/PhysicsAssetEditor/Private/PhysicsAssetEditorActions.cpp b/Engine/Source/Editor/PhysicsAssetEditor/Private/PhysicsAssetEditorActions.cpp
--- a/Engine/Source/Editor/PhysicsAssetEditor/Private/PhysicsAssetEditorActions.cpp
+++ b/Engine/Source/Editor/PhysicsAssetEditor/Private/PhysicsAssetEditorActions.cpp
@@ -61,7 +61,7 @@ void FPhysicsAssetEditorCommands::RegisterCommands()
 	UI_COMMAND(RenderOnlySelectedConstraints, "Only Selected Constraints", "Draw only selected constraints.", EUserInterfaceActionType::ToggleButton, FInputChord());
 	UI_COMMAND(DrawConstraintsAsPoints, "Draw Constraints As Points", "Draw Constraints As Points", EUserInterfaceActionType::ToggleButton, FInputChord());
 	UI_COMMAND(DrawViolatedLimits, "Draw Violated Limits", "Draw Violated Limits", EUserInterfaceActionType::ToggleButton, FInputChord());
-	UI_COMMAND(HideBodyMass, "Hide Mass", "Hide text displaying body mass next to CoM", EUserInterfaceActionType::ToggleButton, FInputChord());
+	UI_COMMAND(DrawBodyMass, "Mass Values", "Draw text displaying body mass next to CoM", EUserInterfaceActionType::ToggleButton, FInputChord());
 	UI_COMMAND(DisableCollision, "Disable Collision", "Disable collision between the currently selected bodies", EUserInterfaceActionType::Button, FInputChord(EKeys::RightBracket));
 	UI_COMMAND(DisableCollisionAll, "Disable Collision All", "Disable collision between the currently selected bodies and all bodies", EUserInterfaceActionType::Button, FInputChord(EModifierKey::Control, EKeys::RightBracket));
 	UI_COMMAND(EnableCollision, "Enable Collision", "Enable collision between the currently selected bodies", EUserInterfaceActionType::Button, FInputChord(EKeys::LeftBracket));
diff --git a/Engine/Source/Editor/PhysicsAssetEditor/Private/PhysicsAssetEditorActions.h b/Engine/Source/Editor/PhysicsAssetEditor/Private/PhysicsAssetEditorActions.h
--- a/Engine/Source/Editor/PhysicsAssetEditor/Private/PhysicsAssetEditorActions.h
+++ b/Engine/Source/Editor/PhysicsAssetEditor/Private/PhysicsAssetEditorActions.h
@@ -70,7 +70,7 @@ public:
 	TSharedPtr<FUICommandInfo> RenderOnlySelectedSolid;
 	TSharedPtr<FUICommandInfo> HideSimulatedBodies;
 	TSharedPtr<FUICommandInfo> HideKinematicBodies;
-	TSharedPtr<FUICommandInfo> HideBodyMass;
+	TSharedPtr<FUICommandInfo> DrawBodyMass;
 	TSharedPtr<FUICommandInfo> DrawCenterOfMass;
 	TSharedPtr<FUICommandInfo> RenderOnlySelectedBodyCoMs;
 	TSharedPtr<FUICommandInfo> RenderOnlySelectedConstraints;
diff --git a/Engine/Source/Editor/PhysicsAssetEditor/Private/PhysicsAssetEditorBodySetupDetailsCustomization.cpp b/Engine/Source/Editor/PhysicsAssetEditor/Private/PhysicsAssetEditorBodySetupDetailsCustomization.cpp
--- a/Engine/Source/Editor/PhysicsAssetEditor/Private/PhysicsAssetEditorBodySetupDetailsCustomization.cpp
+++ b/Engine/Source/Editor/PhysicsAssetEditor/Private/PhysicsAssetEditorBodySetupDetailsCustomization.cpp
@@ -99,6 +99,7 @@ void FPhysicsAssetEditorBodySetupDetailsCustomization::CustomizeCoMNudge(IDetail
 
 				ValueContent->AddSlot()
 					.AutoWidth()
+					.FillWidth(1.0f)
 					.HAlign(HAlign_Left)
 					.VAlign(VAlign_Center)
 					.Padding(HorizontalPadding, VerticalPadding)
@@ -108,11 +109,13 @@ void FPhysicsAssetEditorBodySetupDetailsCustomization::CustomizeCoMNudge(IDetail
 
 				ValueContent->AddSlot()
 					.AutoWidth()
+					.HAlign(HAlign_Center)
 					.VAlign(VAlign_Center)
 					[
 						SNew(SButton)
-							.HAlign(HAlign_Center)
-							.VAlign(VAlign_Center)
+							.ButtonStyle(FAppStyle::Get(), "NoBorder")
+							.ContentPadding(FMargin(0.0f,2.0f))
+							.ContentScale(FVector2D(0.8f, 0.8f))
 							.OnClicked_Lambda([this, Axis]() { return this->ToggleFixCOMInComponentSpace(Axis); })
 							.ButtonColorAndOpacity(FSlateColor::UseForeground())
 							.Content()
diff --git a/Engine/Source/Editor/PhysicsAssetEditor/Private/PhysicsAssetEditorSkeletalMeshComponent.cpp b/Engine/Source/Editor/PhysicsAssetEditor/Private/PhysicsAssetEditorSkeletalMeshComponent.cpp
--- a/Engine/Source/Editor/PhysicsAssetEditor/Private/PhysicsAssetEditorSkeletalMeshComponent.cpp
+++ b/Engine/Source/Editor/PhysicsAssetEditor/Private/PhysicsAssetEditorSkeletalMeshComponent.cpp
@@ -39,7 +39,7 @@ UPhysicsAssetEditorSkeletalMeshComponent::UPhysicsAssetEditorSkeletalMeshCompone
 	, ConstraintBone2Color(0, 150, 150)
 	, HierarchyDrawColor(220, 255, 220)
 	, AnimSkelDrawColor(255, 64, 64)
-	, COMRenderSize(5.0f)
+	, COMRenderSize(2.0f)
 	, InfluenceLineLength(2.0f)
 	, InfluenceLineColor(0, 255, 0)
 {
@@ -120,6 +120,7 @@ void UPhysicsAssetEditorSkeletalMeshComponent::DebugDraw(const FSceneView* View,
 		// when an asset is open in the Physics Asset Editor but should not persist after the editor has been closed.
 		RenderSettings->CenterOfMassViewMode = SharedData->GetCurrentCenterOfMassViewMode(SharedData->bRunningSimulation);
 		RenderSettings->CollisionViewMode = SharedData->GetCurrentCollisionViewMode(SharedData->bRunningSimulation);
+		RenderSettings->COMRenderSize = SharedData->EditorOptions->COMRenderSize;
 		RenderSettings->ConstraintViewMode = SharedData->GetCurrentConstraintViewMode(SharedData->bRunningSimulation);
 		RenderSettings->ConstraintDrawSize = SharedData->EditorOptions->ConstraintDrawSize;
 		RenderSettings->PhysicsBlend = SharedData->EditorOptions->PhysicsBlend;
diff --git a/Engine/Source/Editor/UnrealEd/Classes/Preferences/PhysicsAssetEditorOptions.h b/Engine/Source/Editor/UnrealEd/Classes/Preferences/PhysicsAssetEditorOptions.h
--- a/Engine/Source/Editor/UnrealEd/Classes/Preferences/PhysicsAssetEditorOptions.h
+++ b/Engine/Source/Editor/UnrealEd/Classes/Preferences/PhysicsAssetEditorOptions.h
@@ -127,6 +127,10 @@ class UPhysicsAssetEditorOptions : public UObject
 	UPROPERTY(config)
 	uint32 bSimulationFloorCollisionEnabled:1;
 
+	/** Controls the size of Center of Mass markers drawn in Physics Asset Editor */
+	UPROPERTY(config)
+	float COMRenderSize;
+
 	/** Controls how large constraints are drawn in Physics Asset Editor */
 	UPROPERTY(config)
 	float ConstraintDrawSize;
diff --git a/Engine/Source/Editor/UnrealEd/Private/PreferenceStubs.cpp b/Engine/Source/Editor/UnrealEd/Private/PreferenceStubs.cpp
--- a/Engine/Source/Editor/UnrealEd/Private/PreferenceStubs.cpp
+++ b/Engine/Source/Editor/UnrealEd/Private/PreferenceStubs.cpp
@@ -45,6 +45,7 @@ UPhysicsAssetEditorOptions::UPhysicsAssetEditorOptions(const FObjectInitializer&
 	bShowConstraintsAsPoints = false;
 	bDrawViolatedLimits = false;
 	bSimulationFloorCollisionEnabled = true;
+	COMRenderSize = 1.0f;
 	ConstraintDrawSize = 1.0f;
 
 	// view options
