Add copy name command to physics asset editor, for copying physics bodies names.
#rnx

diff --git a/Engine/Source/Editor/PhysicsAssetEditor/Private/PhysicsAssetEditor.cpp b/Engine/Source/Editor/PhysicsAssetEditor/Private/PhysicsAssetEditor.cpp
--- a/Engine/Source/Editor/PhysicsAssetEditor/Private/PhysicsAssetEditor.cpp
+++ b/Engine/Source/Editor/PhysicsAssetEditor/Private/PhysicsAssetEditor.cpp
@@ -950,6 +950,11 @@ void FPhysicsAssetEditor::BindCommands()
 		FExecuteAction::CreateSP(this, &FPhysicsAssetEditor::OnPasteShapes),
 		FCanExecuteAction::CreateSP(this, &FPhysicsAssetEditor::CanPasteShapes));
 
+	ToolkitCommands->MapAction(
+		Commands.CopyBodyName,
+		FExecuteAction::CreateSP(this, &FPhysicsAssetEditor::OnCopyBodyName),
+		FCanExecuteAction::CreateSP(this, &FPhysicsAssetEditor::CanCopyBodyName));
+		
 	ToolkitCommands->MapAction(
 		Commands.RepeatLastSimulation,
 		FExecuteAction::CreateSP(this, &FPhysicsAssetEditor::OnRepeatLastSimulation),
@@ -1729,6 +1734,7 @@ void FPhysicsAssetEditor::BuildMenuWidgetBody(FMenuBuilder& InMenuBuilder)
 		InMenuBuilder.AddMenuEntry(Commands.PasteShapes);
 		InMenuBuilder.AddMenuEntry(Commands.CopyProperties);
 		InMenuBuilder.AddMenuEntry(Commands.PasteProperties);
+		InMenuBuilder.AddMenuEntry(Commands.CopyBodyName);
 		InMenuBuilder.AddMenuEntry( Commands.DeleteBody );
 		InMenuBuilder.AddMenuEntry( Commands.DeleteAllBodiesBelow );
 		InMenuBuilder.AddMenuEntry( Commands.Mirror );
@@ -2336,6 +2342,21 @@ bool FPhysicsAssetEditor::HasSelectedConstraintAndIsNotSimulation() const
 	return IsNotSimulation() && (SharedData->GetSelectedConstraint());
 }
 
+void FPhysicsAssetEditor::OnCopyBodyName()
+{
+	if(SharedData->SelectedBodies.Num() == 1)
+	{
+		SharedData->CopyBodyName();
+	}
+	
+	RefreshPreviewViewport();
+}
+
+bool FPhysicsAssetEditor::CanCopyBodyName() const
+{
+	return IsSelectedEditMode() && SharedData->SelectedBodies.Num() == 1 && SharedData->SelectedConstraints.Num() == 0;
+}
+
 bool FPhysicsAssetEditor::IsSelectedEditMode() const
 {
 	return HasSelectedBodyAndIsNotSimulation() || HasSelectedConstraintAndIsNotSimulation();
diff --git a/Engine/Source/Editor/PhysicsAssetEditor/Private/PhysicsAssetEditor.h b/Engine/Source/Editor/PhysicsAssetEditor/Private/PhysicsAssetEditor.h
--- a/Engine/Source/Editor/PhysicsAssetEditor/Private/PhysicsAssetEditor.h
+++ b/Engine/Source/Editor/PhysicsAssetEditor/Private/PhysicsAssetEditor.h
@@ -228,6 +228,8 @@ private:
 	bool CanCopyProperties() const;
 	void OnPasteProperties();
 	bool CanPasteProperties() const;
+	void OnCopyBodyName();
+    bool CanCopyBodyName() const;
 	bool IsSelectedEditMode() const;
 	void OnRepeatLastSimulation();
 	void OnToggleSimulation(bool bInSelected);
diff --git a/Engine/Source/Editor/PhysicsAssetEditor/Private/PhysicsAssetEditorActions.cpp b/Engine/Source/Editor/PhysicsAssetEditor/Private/PhysicsAssetEditorActions.cpp
--- a/Engine/Source/Editor/PhysicsAssetEditor/Private/PhysicsAssetEditorActions.cpp
+++ b/Engine/Source/Editor/PhysicsAssetEditor/Private/PhysicsAssetEditorActions.cpp
@@ -21,6 +21,7 @@ void FPhysicsAssetEditorCommands::RegisterCommands()
 	UI_COMMAND(PasteShapes, "Paste Selected Shapes From Clipboard", "Paste shapes from clipboard to selected bodies", EUserInterfaceActionType::Button, FInputChord());
 	UI_COMMAND(CopyProperties, "Copy Properties", "Copy Properties: Copy Properties Of Currently Selected Object To Next Selected Object", EUserInterfaceActionType::Button, FInputChord(EModifierKey::Control, EKeys::C));
 	UI_COMMAND(PasteProperties, "Paste Properties", "Paste Properties: Copy Properties Of Currently Selected Object To Next Selected Object", EUserInterfaceActionType::Button, FInputChord(EModifierKey::Control, EKeys::V));
+	UI_COMMAND(CopyBodyName, "Copy Body Name", "Copy Body Name: Copy name of currently selected body to clipboard", EUserInterfaceActionType::Button, FInputChord());
 	UI_COMMAND(RepeatLastSimulation, "Simulate", "Previews Physics Simulation", EUserInterfaceActionType::RadioButton, FInputChord());
 	UI_COMMAND(SimulationNoGravity, "No Gravity Simulation", "Run Physics Simulation without gravity. Use this to debug issues with your ragdoll. If the setup is correct, the asset should not move!", EUserInterfaceActionType::ToggleButton, FInputChord());
 	UI_COMMAND(SimulationFloorCollision, "Enable Floor Collisions", "Run Physics Simulation with collisions between physics bodies and the floor enabled.", EUserInterfaceActionType::ToggleButton, FInputChord());
diff --git a/Engine/Source/Editor/PhysicsAssetEditor/Private/PhysicsAssetEditorActions.h b/Engine/Source/Editor/PhysicsAssetEditor/Private/PhysicsAssetEditorActions.h
--- a/Engine/Source/Editor/PhysicsAssetEditor/Private/PhysicsAssetEditorActions.h
+++ b/Engine/Source/Editor/PhysicsAssetEditor/Private/PhysicsAssetEditorActions.h
@@ -36,6 +36,7 @@ public:
 	TSharedPtr<FUICommandInfo> PasteShapes;
 	TSharedPtr<FUICommandInfo> CopyProperties;
 	TSharedPtr<FUICommandInfo> PasteProperties;
+	TSharedPtr<FUICommandInfo> CopyBodyName;
 	TSharedPtr<FUICommandInfo> RepeatLastSimulation;
 	TSharedPtr<FUICommandInfo> SimulationNoGravity;
 	TSharedPtr<FUICommandInfo> SimulationFloorCollision;
diff --git a/Engine/Source/Editor/PhysicsAssetEditor/Private/PhysicsAssetEditorSharedData.cpp b/Engine/Source/Editor/PhysicsAssetEditor/Private/PhysicsAssetEditorSharedData.cpp
--- a/Engine/Source/Editor/PhysicsAssetEditor/Private/PhysicsAssetEditorSharedData.cpp
+++ b/Engine/Source/Editor/PhysicsAssetEditor/Private/PhysicsAssetEditorSharedData.cpp
@@ -2094,6 +2094,12 @@ void FPhysicsAssetEditorSharedData::PasteBodyProperties()
 	}
 }
 
+void FPhysicsAssetEditorSharedData::CopyBodyName()
+{
+	check(SelectedBodies.Num() == 1);
+	FPlatformApplicationMisc::ClipboardCopy(*PhysicsAsset->SkeletalBodySetups[GetSelectedBody()->Index]->BoneName.ToString());
+}
+
 bool FPhysicsAssetEditorSharedData::WeldSelectedBodies(bool bWeld /* = true */)
 {
 	bool bCanWeld = false;
diff --git a/Engine/Source/Editor/PhysicsAssetEditor/Private/PhysicsAssetEditorSharedData.h b/Engine/Source/Editor/PhysicsAssetEditor/Private/PhysicsAssetEditorSharedData.h
--- a/Engine/Source/Editor/PhysicsAssetEditor/Private/PhysicsAssetEditorSharedData.h
+++ b/Engine/Source/Editor/PhysicsAssetEditor/Private/PhysicsAssetEditorSharedData.h
@@ -215,6 +215,7 @@ public:
 	void CopyBodyProperties();
 	void CopyConstraintProperties();
 	void PasteBodyProperties();
+	void CopyBodyName();
 	bool WeldSelectedBodies(bool bWeld = true);
 	void Mirror();
 
