[UE]Phat: Fix wording and UI for creating bodies
- Change the name and tool tips for the 'Add/Replace Bodies' entry in the PhAT Skeleton Tree context menu for bones. The name text is now dependant on the nature of the selected bones and describes how the menu option will affect those bones.
- Remove the 'Add Shapes' entry from the PhAT Skeleton Tree context menu for bones as it was confusing (see jira).
- Renamed the "Add Shape" mnu entry to "Add Primitive" to make it more consistent with the rest of the UI.

diff --git a/Engine/Source/Editor/PhysicsAssetEditor/Private/PhysicsAssetEditor.cpp b/Engine/Source/Editor/PhysicsAssetEditor/Private/PhysicsAssetEditor.cpp
--- a/Engine/Source/Editor/PhysicsAssetEditor/Private/PhysicsAssetEditor.cpp
+++ b/Engine/Source/Editor/PhysicsAssetEditor/Private/PhysicsAssetEditor.cpp
@@ -593,6 +593,58 @@ FSlateIcon FPhysicsAssetEditor::GetRepeatLastSimulationIcon() const
 	}
 }
 
+void FPhysicsAssetEditor::BuildMenuWidgetBone(FMenuBuilder& InMenuBuilder, const TSharedPtr< const FUICommandInfo > InCommand)
+{
+	InMenuBuilder.PushCommandList(GetToolkitCommands());
+	InMenuBuilder.BeginSection("BodyActions", LOCTEXT("BodyHeader", "Body"));
+	{
+		InMenuBuilder.AddMenuEntry(InCommand);
+	}
+	InMenuBuilder.EndSection();
+	AddAdvancedMenuWidget(InMenuBuilder);
+	InMenuBuilder.PopCommandList();
+}
+
+void FPhysicsAssetEditor::BuildMenuWidgetBone(FMenuBuilder& InMenuBuilder, const TArray<TSharedPtr<ISkeletonTreeItem>>* const SelectedBones)
+{
+	const FPhysicsAssetEditorCommands& Commands = FPhysicsAssetEditorCommands::Get();
+
+	TSharedPtr< const FUICommandInfo > MenuEntryCommand = Commands.CreateOrRegenerateBodies;
+
+	// Determine whether any or all selected bones have child bodies.
+	if (SelectedBones)
+	{
+		uint32 SelectedBoneCount = 0;
+		uint32 WithBodyBoneCount = 0;
+
+		for (const TSharedPtr<ISkeletonTreeItem>& SkeletonTreeItem : *SelectedBones)
+		{
+			++SelectedBoneCount;
+
+			for (const TSharedPtr<ISkeletonTreeItem>& ChildTreeItem : SkeletonTreeItem->GetChildren())
+			{
+				if (ChildTreeItem->IsOfType<FSkeletonTreePhysicsBodyItem>())
+				{
+					++WithBodyBoneCount;
+					break;
+				}
+			}
+		}
+
+		// Find the most appropriate menu entry for the selected bones.
+		if (WithBodyBoneCount == 0)
+		{
+			MenuEntryCommand = Commands.CreateBodies;
+		}
+		else if (SelectedBoneCount == WithBodyBoneCount)
+		{
+			MenuEntryCommand = Commands.RegenerateBodies;
+		}
+	}
+
+	BuildMenuWidgetBone(InMenuBuilder, MenuEntryCommand);
+}
+
 void FPhysicsAssetEditor::ExtendToolbar()
 {
 	struct Local
@@ -921,7 +973,12 @@ void FPhysicsAssetEditor::BindCommands()
 		FCanExecuteAction::CreateSP(this, &FPhysicsAssetEditor::IsNotSimulation));
 
 	ToolkitCommands->MapAction(
-		Commands.AddBodies,
+		Commands.CreateBodies,
+		FExecuteAction::CreateSP(this, &FPhysicsAssetEditor::ResetBoneCollision),
+		FCanExecuteAction::CreateSP(this, &FPhysicsAssetEditor::IsNotSimulation));
+
+	ToolkitCommands->MapAction(
+		Commands.CreateOrRegenerateBodies,
 		FExecuteAction::CreateSP(this, &FPhysicsAssetEditor::ResetBoneCollision),
 		FCanExecuteAction::CreateSP(this, &FPhysicsAssetEditor::IsNotSimulation));
 
@@ -1714,7 +1771,7 @@ void FPhysicsAssetEditor::BuildMenuWidgetBody(FMenuBuilder& InMenuBuilder)
 
 		InMenuBuilder.BeginSection( "BodyActions", LOCTEXT( "BodyHeader", "Body" ) );
 		InMenuBuilder.AddMenuEntry( Commands.RegenerateBodies );
-		InMenuBuilder.AddSubMenu( LOCTEXT("AddShapeMenu", "Add Shape"), LOCTEXT("AddShapeMenu_ToolTip", "Add shapes to this body"),
+		InMenuBuilder.AddSubMenu( LOCTEXT("AddPrimitiveMenu", "Add Primitive"), LOCTEXT("AddPrimitiveMenu_ToolTip", "Add Primitives to this body"),
 			FNewMenuDelegate::CreateStatic( &FillAddShapeMenu ) );
 		InMenuBuilder.AddSubMenu( LOCTEXT("CollisionMenu", "Collision"), LOCTEXT("CollisionMenu_ToolTip", "Adjust body/body collision"),
 			FNewMenuDelegate::CreateStatic( &FLocal::FillCollisionMenu ) );	
@@ -1985,17 +2042,7 @@ TSharedRef<ISkeletonTree> FPhysicsAssetEditor::BuildMenuWidgetNewConstraintForBo
 
 void FPhysicsAssetEditor::BuildMenuWidgetBone(FMenuBuilder& InMenuBuilder)
 {
-	InMenuBuilder.PushCommandList(GetToolkitCommands());
-	InMenuBuilder.BeginSection( "BodyActions", LOCTEXT( "BodyHeader", "Body" ) );
-	{
-		const FPhysicsAssetEditorCommands& Commands = FPhysicsAssetEditorCommands::Get();
-		InMenuBuilder.AddMenuEntry( Commands.AddBodies );
-		InMenuBuilder.AddSubMenu( LOCTEXT("AddShapeMenu", "Add Shape"), LOCTEXT("AddShapeMenu_ToolTip", "Add shapes to this body"),
-			FNewMenuDelegate::CreateStatic( &FillAddShapeMenu ) );
-	}
-	InMenuBuilder.EndSection();
-	AddAdvancedMenuWidget(InMenuBuilder);
-	InMenuBuilder.PopCommandList();
+	BuildMenuWidgetBone(InMenuBuilder, FPhysicsAssetEditorCommands::Get().CreateBodies);
 }
 
 bool FPhysicsAssetEditor::ShouldFilterAssetBasedOnSkeleton( const FAssetData& AssetData )
@@ -4209,6 +4256,7 @@ void FPhysicsAssetEditor::HandleExtendContextMenu(FMenuBuilder& InMenuBuilder)
 	TArray<TSharedPtr<FSkeletonTreePhysicsConstraintItem>> SelectedConstraints = Selection.GetSelectedItems<FSkeletonTreePhysicsConstraintItem>();
 	TArray<TSharedPtr<FSkeletonTreePhysicsShapeItem>> SelectedShapes = Selection.GetSelectedItems<FSkeletonTreePhysicsShapeItem>();
 	TArray<TSharedPtr<ISkeletonTreeItem>> SelectedBones = Selection.GetSelectedItemsByTypeId("FSkeletonTreeBoneItem");
+	
 	if (SelectedBodies.Num() > 0)
 	{
 		BuildMenuWidgetBody(InMenuBuilder);
@@ -4223,7 +4271,7 @@ void FPhysicsAssetEditor::HandleExtendContextMenu(FMenuBuilder& InMenuBuilder)
 	}
 	else if(SelectedBones.Num() > 0)
 	{
-		BuildMenuWidgetBone(InMenuBuilder);
+		BuildMenuWidgetBone(InMenuBuilder, &SelectedBones);
 	}
 
 	BuildMenuWidgetSelection(InMenuBuilder);
diff --git a/Engine/Source/Editor/PhysicsAssetEditor/Private/PhysicsAssetEditor.h b/Engine/Source/Editor/PhysicsAssetEditor/Private/PhysicsAssetEditor.h
--- a/Engine/Source/Editor/PhysicsAssetEditor/Private/PhysicsAssetEditor.h
+++ b/Engine/Source/Editor/PhysicsAssetEditor/Private/PhysicsAssetEditor.h
@@ -335,6 +335,12 @@ private:
 	FText GetRepeatLastSimulationToolTip() const;
 	FSlateIcon GetRepeatLastSimulationIcon() const;
 
+	/** Add a specified skeleton tree context menu entriy specific to bones. */
+	void BuildMenuWidgetBone(FMenuBuilder& InMenuBuilder, const TSharedPtr< const FUICommandInfo > InCommand);
+
+	/** Add skeleton tree context menu entries specific to the currently selected set of bones. */
+	void BuildMenuWidgetBone(FMenuBuilder& InMenuBuilder, const TArray<TSharedPtr<ISkeletonTreeItem>>* const SelectedBones);
+
 	/** Handle initial preview scene setup */
 	void HandlePreviewSceneCreated(const TSharedRef<IPersonaPreviewScene>& InPersonaPreviewScene);
 
diff --git a/Engine/Source/Editor/PhysicsAssetEditor/Private/PhysicsAssetEditorActions.cpp b/Engine/Source/Editor/PhysicsAssetEditor/Private/PhysicsAssetEditorActions.cpp
--- a/Engine/Source/Editor/PhysicsAssetEditor/Private/PhysicsAssetEditorActions.cpp
+++ b/Engine/Source/Editor/PhysicsAssetEditor/Private/PhysicsAssetEditorActions.cpp
@@ -12,7 +12,8 @@
 void FPhysicsAssetEditorCommands::RegisterCommands()
 {
 	UI_COMMAND(RegenerateBodies, "Regenerate Bodies", "Regenerates the selected bodies using the current generation settings (see the Tools tab)", EUserInterfaceActionType::Button, FInputChord());
-	UI_COMMAND(AddBodies, "Add/Replace Bodies", "Adds or replaces bodies for the selected bones using the current generation settings (see the Tools tab)", EUserInterfaceActionType::Button, FInputChord());
+	UI_COMMAND(CreateBodies, "Create Bodies/Constraints", "Creates bodies and/or constraints for the selected bones using the current generation settings (see the Tools tab)", EUserInterfaceActionType::Button, FInputChord());
+	UI_COMMAND(CreateOrRegenerateBodies, "Create/Regenerate Bodies", "Creates bodies and/or constraints (if none exist) or Regenerates bodies (if they do exist) for each of the selected bones using the current generation settings (see the Tools tab)", EUserInterfaceActionType::Button, FInputChord());
 
 	UI_COMMAND(ApplyPhysicalMaterial, "Physical Material", "Apply a Physical Material to All Bodies", EUserInterfaceActionType::Button, FInputChord());
 	UI_COMMAND(CopyBodies, "Copy Selected Bodies/Constraints To Clipboard", "Copy selected bodies/constraints to clipboard", EUserInterfaceActionType::Button, FInputChord(EModifierKey::Control | EModifierKey::Shift, EKeys::C));
diff --git a/Engine/Source/Editor/PhysicsAssetEditor/Private/PhysicsAssetEditorActions.h b/Engine/Source/Editor/PhysicsAssetEditor/Private/PhysicsAssetEditorActions.h
--- a/Engine/Source/Editor/PhysicsAssetEditor/Private/PhysicsAssetEditorActions.h
+++ b/Engine/Source/Editor/PhysicsAssetEditor/Private/PhysicsAssetEditorActions.h
@@ -27,7 +27,8 @@ public:
 	/** See tooltips in cpp for documentation */
 	
 	TSharedPtr<FUICommandInfo> RegenerateBodies;
-	TSharedPtr<FUICommandInfo> AddBodies;
+	TSharedPtr<FUICommandInfo> CreateBodies;
+	TSharedPtr<FUICommandInfo> CreateOrRegenerateBodies;
 	TSharedPtr<FUICommandInfo> ApplyPhysicalMaterial;
 	TSharedPtr<FUICommandInfo> Snap;
 	TSharedPtr<FUICommandInfo> CopyBodies;
