[UE] Physics Asset Tool Skeleton Tree Body/Primitive options
- New sub menu when right-clicking on a bone that doesn't have anything yet.
- Add primitive submenu when right-clicking on a bone that has a body.
- Constraints are optionall created with bodies, determined by a checkbox that is stored per user and persists in user settings between editor sessions.
- Present inactive/grayed out menu options when bone selection contains a mixture of bones with and without associated bodies.

diff --git a/Engine/Source/Editor/EditorStyle/Private/SlateEditorStyle.cpp b/Engine/Source/Editor/EditorStyle/Private/SlateEditorStyle.cpp
--- a/Engine/Source/Editor/EditorStyle/Private/SlateEditorStyle.cpp
+++ b/Engine/Source/Editor/EditorStyle/Private/SlateEditorStyle.cpp
@@ -7578,6 +7578,11 @@ void FSlateEditorStyle::FStyle::SetupToolkitStyles()
 		Set( "PhysicsAssetEditor.AddSphyl", new IMAGE_BRUSH( "PhysicsAssetEditor/icon_PhAT_Sphyl_40x", Icon40x40 ) );
 		Set( "PhysicsAssetEditor.AddBox", new IMAGE_BRUSH( "PhysicsAssetEditor/icon_PhAT_Box_40x", Icon40x40 ) );
 		Set( "PhysicsAssetEditor.AddTaperedCapsule", new IMAGE_BRUSH( "PhysicsAssetEditor/icon_PhAT_TaperedCapsule_40x", Icon40x40 ) );
+		Set( "PhysicsAssetEditor.CreateBodyWithSphere", new IMAGE_BRUSH("PhysicsAssetEditor/icon_PhAT_Sphere_40x", Icon40x40));
+		Set( "PhysicsAssetEditor.CreateBodyWithSphyl", new IMAGE_BRUSH("PhysicsAssetEditor/icon_PhAT_Sphyl_40x", Icon40x40));
+		Set( "PhysicsAssetEditor.CreateBodyWithBox", new IMAGE_BRUSH("PhysicsAssetEditor/icon_PhAT_Box_40x", Icon40x40));
+		Set( "PhysicsAssetEditor.CreateBodyWithTaperedCapsule", new IMAGE_BRUSH("PhysicsAssetEditor/icon_PhAT_TaperedCapsule_40x", Icon40x40));
+		Set( "PhysicsAssetEditor.CreateBodyShouldCreateConstraints", new IMAGE_BRUSH("/PhysicsAssetEditor/Constraint_16x", Icon40x40));
 		Set( "PhysicsAssetEditor.DeletePrimitive", new IMAGE_BRUSH( "PhysicsAssetEditor/icon_PhAT_DeletePrimitive_40x", Icon40x40 ) );
 		Set( "PhysicsAssetEditor.DuplicatePrimitive", new IMAGE_BRUSH( "PhysicsAssetEditor/icon_PhAT_DupePrim_40x", Icon40x40 ) );
 		Set( "PhysicsAssetEditor.ResetConstraint", new IMAGE_BRUSH( "PhysicsAssetEditor/icon_PhAT_ResetConstraint_40x", Icon40x40 ) );
@@ -7610,6 +7615,11 @@ void FSlateEditorStyle::FStyle::SetupToolkitStyles()
 		Set( "PhysicsAssetEditor.DisableCollisionAll.Small", new IMAGE_BRUSH( "PhysicsAssetEditor/icon_PhAT_DisableCollision_40x", Icon20x20 ) );
 		Set( "PhysicsAssetEditor.EnableCollisionAll.Small", new IMAGE_BRUSH( "PhysicsAssetEditor/icon_PhAT_EnableCollision_40x", Icon20x20 ) );
 		Set( "PhysicsAssetEditor.WeldToBody.Small", new IMAGE_BRUSH( "PhysicsAssetEditor/icon_PhAT_Weld_40x", Icon20x20 ) );
+		Set( "PhysicsAssetEditor.CreateBodyWithSphere.Small", new IMAGE_BRUSH("PhysicsAssetEditor/icon_PhAT_Sphere_40x", Icon20x20));
+		Set( "PhysicsAssetEditor.CreateBodyWithSphyl.Small", new IMAGE_BRUSH("PhysicsAssetEditor/icon_PhAT_Sphyl_40x", Icon20x20));
+		Set( "PhysicsAssetEditor.CreateBodyWithBox.Small", new IMAGE_BRUSH("PhysicsAssetEditor/icon_PhAT_Box_40x", Icon20x20));
+		Set( "PhysicsAssetEditor.CreateBodyWithTaperedCapsule.Small", new IMAGE_BRUSH("PhysicsAssetEditor/icon_PhAT_TaperedCapsule_40x", Icon20x20));
+		Set( "PhysicsAssetEditor.CreateBodyShouldCreateConstraints.Small", new IMAGE_BRUSH("/PhysicsAssetEditor/Constraint_16x", Icon20x20));
 		Set( "PhysicsAssetEditor.AddNewBody.Small", new IMAGE_BRUSH( "PhysicsAssetEditor/icon_PhAT_NewBody_40x", Icon20x20 ) );
 		Set( "PhysicsAssetEditor.AddSphere.Small", new IMAGE_BRUSH( "PhysicsAssetEditor/icon_PhAT_Sphere_40x", Icon20x20 ) );
 		Set( "PhysicsAssetEditor.AddSphyl.Small", new IMAGE_BRUSH( "PhysicsAssetEditor/icon_PhAT_Sphyl_40x", Icon20x20 ) );
diff --git a/Engine/Source/Editor/EditorStyle/Private/StarshipStyle.cpp b/Engine/Source/Editor/EditorStyle/Private/StarshipStyle.cpp
--- a/Engine/Source/Editor/EditorStyle/Private/StarshipStyle.cpp
+++ b/Engine/Source/Editor/EditorStyle/Private/StarshipStyle.cpp
@@ -7488,6 +7488,11 @@ void FStarshipEditorStyle::FStyle::SetupToolkitStyles()
 		Set( "PhysicsAssetEditor.DisableCollisionAll", new IMAGE_BRUSH_SVG("Starship/Common/DisableCollision", Icon20x20));
 		Set( "PhysicsAssetEditor.EnableCollisionAll", new IMAGE_BRUSH_SVG("Starship/Common/EnableCollision", Icon20x20));
 		Set( "PhysicsAssetEditor.WeldToBody", new IMAGE_BRUSH( "PhysicsAssetEditor/icon_PhAT_Weld_40x", Icon40x40 ) );
+		Set( "PhysicsAssetEditor.CreateBodyWithSphere", new IMAGE_BRUSH("PhysicsAssetEditor/icon_PhAT_Sphere_40x", Icon40x40));
+		Set( "PhysicsAssetEditor.CreateBodyWithSphyl", new IMAGE_BRUSH("PhysicsAssetEditor/icon_PhAT_Sphyl_40x", Icon40x40));
+		Set( "PhysicsAssetEditor.CreateBodyWithBox", new IMAGE_BRUSH("PhysicsAssetEditor/icon_PhAT_Box_40x", Icon40x40));
+		Set( "PhysicsAssetEditor.CreateBodyWithTaperedCapsule", new IMAGE_BRUSH("PhysicsAssetEditor/icon_PhAT_TaperedCapsule_40x", Icon40x40));
+		Set( "PhysicsAssetEditor.CreateBodyShouldCreateConstraints", new IMAGE_BRUSH("PhysicsAssetEditor/Constraint_16x", Icon40x40));
 		Set( "PhysicsAssetEditor.AddNewBody", new IMAGE_BRUSH( "PhysicsAssetEditor/icon_PhAT_NewBody_40x", Icon40x40 ) );
 		Set( "PhysicsAssetEditor.AddSphere", new IMAGE_BRUSH( "PhysicsAssetEditor/icon_PhAT_Sphere_40x", Icon40x40 ) );
 		Set( "PhysicsAssetEditor.AddSphyl", new IMAGE_BRUSH( "PhysicsAssetEditor/icon_PhAT_Sphyl_40x", Icon40x40 ) );
diff --git a/Engine/Source/Editor/PhysicsAssetEditor/Private/PhysicsAssetEditor.cpp b/Engine/Source/Editor/PhysicsAssetEditor/Private/PhysicsAssetEditor.cpp
--- a/Engine/Source/Editor/PhysicsAssetEditor/Private/PhysicsAssetEditor.cpp
+++ b/Engine/Source/Editor/PhysicsAssetEditor/Private/PhysicsAssetEditor.cpp
@@ -170,11 +170,11 @@ FPhysicsAssetEditor::~FPhysicsAssetEditor()
 	}
 }
 
-static void FillAddShapeMenu(FMenuBuilder& InSubMenuBuilder)
+static void FillAddPrimitiveMenu(FMenuBuilder& InSubMenuBuilder)
 {
 	const FPhysicsAssetEditorCommands& PhysicsAssetEditorCommands = FPhysicsAssetEditorCommands::Get();
 
-	InSubMenuBuilder.BeginSection("ShapeTypeHeader", LOCTEXT("ShapeTypeHeader", "Shape Type"));
+	InSubMenuBuilder.BeginSection("PrimitiveTypeHeader", LOCTEXT("PrimitiveTypeHeader", "Primitive Type"));
 	InSubMenuBuilder.AddMenuEntry( PhysicsAssetEditorCommands.AddBox );
 	InSubMenuBuilder.AddMenuEntry( PhysicsAssetEditorCommands.AddSphere );
 	InSubMenuBuilder.AddMenuEntry( PhysicsAssetEditorCommands.AddSphyl );
@@ -182,6 +182,28 @@ static void FillAddShapeMenu(FMenuBuilder& InSubMenuBuilder)
 	InSubMenuBuilder.EndSection();
 }
 
+static void FillCreateBodiesConstraintsMenu(FMenuBuilder& InSubMenuBuilder)
+{
+	const FPhysicsAssetEditorCommands& PhysicsAssetEditorCommands = FPhysicsAssetEditorCommands::Get();
+
+	// Primitive Type Section
+	{
+		InSubMenuBuilder.BeginSection("PrimitiveTypeHeader", LOCTEXT("PrimitiveTypeHeader", "Primitive Type"));
+		InSubMenuBuilder.AddMenuEntry( PhysicsAssetEditorCommands.CreateBodyWithBox );
+		InSubMenuBuilder.AddMenuEntry( PhysicsAssetEditorCommands.CreateBodyWithSphere );
+		InSubMenuBuilder.AddMenuEntry( PhysicsAssetEditorCommands.CreateBodyWithSphyl );
+		InSubMenuBuilder.AddMenuEntry( PhysicsAssetEditorCommands.CreateBodyWithTaperedCapsule );
+		InSubMenuBuilder.EndSection();
+	}
+
+	// Advanced Section
+	{
+		InSubMenuBuilder.BeginSection("AdvancedHeader", LOCTEXT("AdvancedHeader", "Advanced"));
+		InSubMenuBuilder.AddMenuEntry( PhysicsAssetEditorCommands.CreateBodyShouldCreateConstraints);
+		InSubMenuBuilder.EndSection();
+	}
+}
+
 void FPhysicsAssetEditor::InitPhysicsAssetEditor(const EToolkitMode::Type Mode, const TSharedPtr< class IToolkitHost >& InitToolkitHost, UPhysicsAsset* ObjectToEdit)
 {
 	SelectedSimulation = false;
@@ -217,7 +239,7 @@ void FPhysicsAssetEditor::InitPhysicsAssetEditor(const EToolkitMode::Type Mode,
 	SkeletonTreeArgs.Extenders = MakeShared<FExtender>();
 	SkeletonTreeArgs.Extenders->AddMenuExtension("FilterOptions", EExtensionHook::After, GetToolkitCommands(), FMenuExtensionDelegate::CreateSP(this, &FPhysicsAssetEditor::HandleExtendFilterMenu));
 	SkeletonTreeArgs.Extenders->AddMenuExtension("SkeletonTreeContextMenu", EExtensionHook::After, GetToolkitCommands(), FMenuExtensionDelegate::CreateSP(this, &FPhysicsAssetEditor::HandleExtendContextMenu));
-	SkeletonTreeArgs.Extenders->AddMenuExtension("CreateNew", EExtensionHook::After, GetToolkitCommands(), FMenuExtensionDelegate::CreateStatic( &FillAddShapeMenu));
+	SkeletonTreeArgs.Extenders->AddMenuExtension("CreateNew", EExtensionHook::After, GetToolkitCommands(), FMenuExtensionDelegate::CreateStatic( &FillAddPrimitiveMenu));
 	SkeletonTreeArgs.Builder = SkeletonTreeBuilder = MakeShared<FPhysicsAssetEditorSkeletonTreeBuilder>(ObjectToEdit);
 	SkeletonTreeArgs.ContextName = GetToolkitFName();
 
@@ -593,56 +615,84 @@ FSlateIcon FPhysicsAssetEditor::GetRepeatLastSimulationIcon() const
 	}
 }
 
-void FPhysicsAssetEditor::BuildMenuWidgetBone(FMenuBuilder& InMenuBuilder, const TSharedPtr< const FUICommandInfo > InCommand)
-{
-	InMenuBuilder.PushCommandList(GetToolkitCommands());
-	InMenuBuilder.BeginSection("BodyActions", LOCTEXT("BodyHeader", "Body"));
-	{
-		InMenuBuilder.AddMenuEntry(InCommand);
-	}
-	InMenuBuilder.EndSection();
-	AddAdvancedMenuWidget(InMenuBuilder);
-	InMenuBuilder.PopCommandList();
-}
-
-void FPhysicsAssetEditor::BuildMenuWidgetBone(FMenuBuilder& InMenuBuilder, const TArray<TSharedPtr<ISkeletonTreeItem>>* const SelectedBones)
+void FPhysicsAssetEditor::BuildMenuWidgetBone(FMenuBuilder& InMenuBuilder, const TArray<TSharedPtr<ISkeletonTreeItem>>& SelectedBones)
 {
 	const FPhysicsAssetEditorCommands& Commands = FPhysicsAssetEditorCommands::Get();
 
 	TSharedPtr< const FUICommandInfo > MenuEntryCommand = Commands.CreateOrRegenerateBodies;
 
 	// Determine whether any or all selected bones have child bodies.
-	if (SelectedBones)
+	uint32 WithBodyBoneCount = 0;
+	uint32 SelectedBoneCount = 0;
+
+	for (const TSharedPtr<ISkeletonTreeItem>& SkeletonTreeItem : SelectedBones)
 	{
-		uint32 SelectedBoneCount = 0;
-		uint32 WithBodyBoneCount = 0;
+		++SelectedBoneCount;
 
-		for (const TSharedPtr<ISkeletonTreeItem>& SkeletonTreeItem : *SelectedBones)
+		for (const TSharedPtr<ISkeletonTreeItem>& ChildTreeItem : SkeletonTreeItem->GetChildren())
 		{
-			++SelectedBoneCount;
-
-			for (const TSharedPtr<ISkeletonTreeItem>& ChildTreeItem : SkeletonTreeItem->GetChildren())
+			if (ChildTreeItem->IsOfType<FSkeletonTreePhysicsBodyItem>())
 			{
-				if (ChildTreeItem->IsOfType<FSkeletonTreePhysicsBodyItem>())
-				{
-					++WithBodyBoneCount;
-					break;
-				}
+				++WithBodyBoneCount;
+				break;
 			}
 		}
+	}
 
-		// Find the most appropriate menu entry for the selected bones.
-		if (WithBodyBoneCount == 0)
-		{
-			MenuEntryCommand = Commands.CreateBodies;
-		}
-		else if (SelectedBoneCount == WithBodyBoneCount)
-		{
-			MenuEntryCommand = Commands.RegenerateBodies;
-		}
+	static const FText AddPrimitiveMenuText = LOCTEXT("AddPrimitiveMenu", "Add Primitive");
+	static const FText CreateBodiesConstraintsMenuText = LOCTEXT("CreateBodiesConstraintsMenu", "Create Bodies / Constraints");
+
+	// Build sub menu.
+	InMenuBuilder.PushCommandList(GetToolkitCommands());
+	InMenuBuilder.BeginSection("BodyActions", LOCTEXT("BodyHeader", "Body"));
+
+	// Find the most appropriate menu entry for the selected bones.
+	if (WithBodyBoneCount == 0)
+	{
+		// None of the selected bones have bodies - show the sub menu to create new bodies using a specified primitive with or without constraints.
+		InMenuBuilder.AddSubMenu(CreateBodiesConstraintsMenuText, LOCTEXT("CreateBodiesConstraintsMenu_ToolTip", "Create new bodies and potentially constraints for the selected bones."),
+			FNewMenuDelegate::CreateStatic(&FillCreateBodiesConstraintsMenu));
+	}
+	else if(WithBodyBoneCount == SelectedBoneCount)
+	{
+		// All of the selected bones have bodies - show the options to regenerate bodies and add primitives.
+		InMenuBuilder.AddMenuEntry(Commands.RegenerateBodies);
+		InMenuBuilder.AddSubMenu(AddPrimitiveMenuText, LOCTEXT("AddPrimitiveMenu_ToolTip", "Add Primitives to this body"),
+			FNewMenuDelegate::CreateStatic(&FillAddPrimitiveMenu));
+	}
+	else
+	{
+		// Some selected bones have bodies, some do not - gray out all the options and use tool tips to explain why.
+		const FText InconsistentSelectionForNoBodyCommandsToolTip = LOCTEXT("InconsistentBoneSelectionForNoBodyCommandsToolTip", "This option is only available when none of the selected bones have an associated body.");
+		const FText InconsistentSelectionForWithBodyCommandsToolTip = LOCTEXT("InconsistentBoneSelectionForWithBodyCommandsToolTip", "This option is only available when all of the selected bones have an associated body.");
+
+		const TSharedRef< SWidget > CreateBodiesConstraintsWidget = 
+			SNew(STextBlock)
+			.Text(CreateBodiesConstraintsMenuText)
+			.TextStyle(FAppStyle::Get(), "NormalText.Subdued")
+			.ToolTipText(InconsistentSelectionForNoBodyCommandsToolTip);
+
+		const TSharedRef< SWidget > AddPrimitiveWidget =
+			SNew(STextBlock)
+			.Text(AddPrimitiveMenuText)
+			.TextStyle(FAppStyle::Get(), "NormalText.Subdued")
+			.ToolTipText(InconsistentSelectionForWithBodyCommandsToolTip);
+
+		const TSharedRef< SWidget > RegenerateBodiesWidget =
+			SNew(STextBlock)
+			.Text(Commands.RegenerateBodies->GetLabel())
+			.TextStyle(FAppStyle::Get(), "NormalText.Subdued")
+			.ToolTipText(InconsistentSelectionForWithBodyCommandsToolTip);
+
+		// Add non-functional menu entries that appear grayed-out.
+		InMenuBuilder.AddMenuEntry(FUIAction(), RegenerateBodiesWidget);
+		InMenuBuilder.AddMenuEntry(FUIAction(), AddPrimitiveWidget);
+		InMenuBuilder.AddMenuEntry(FUIAction(), CreateBodiesConstraintsWidget);
 	}
 
-	BuildMenuWidgetBone(InMenuBuilder, MenuEntryCommand);
+	InMenuBuilder.EndSection();
+	AddAdvancedMenuWidget(InMenuBuilder);
+	InMenuBuilder.PopCommandList();	
 }
 
 void FPhysicsAssetEditor::ExtendToolbar()
@@ -977,6 +1027,32 @@ void FPhysicsAssetEditor::BindCommands()
 		FExecuteAction::CreateSP(this, &FPhysicsAssetEditor::ResetBoneCollision),
 		FCanExecuteAction::CreateSP(this, &FPhysicsAssetEditor::IsNotSimulation));
 
+	ToolkitCommands->MapAction(
+		Commands.CreateBodyWithBox,
+		FExecuteAction::CreateSPLambda(this, [this]() { this->CreateBodiesAndConstraintsForSelectedBones(EAggCollisionShape::Box, this->ShouldCreateConstraintsWhenCreatingBodies()); }),
+		FCanExecuteAction::CreateSP(this, &FPhysicsAssetEditor::IsNotSimulation));
+
+	ToolkitCommands->MapAction(
+		Commands.CreateBodyWithSphere,
+		FExecuteAction::CreateSPLambda(this, [this]() { this->CreateBodiesAndConstraintsForSelectedBones(EAggCollisionShape::Sphere, this->ShouldCreateConstraintsWhenCreatingBodies()); }),
+		FCanExecuteAction::CreateSP(this, &FPhysicsAssetEditor::IsNotSimulation));
+
+	ToolkitCommands->MapAction(
+		Commands.CreateBodyWithSphyl,
+		FExecuteAction::CreateSPLambda(this, [this]() { this->CreateBodiesAndConstraintsForSelectedBones(EAggCollisionShape::Sphyl, this->ShouldCreateConstraintsWhenCreatingBodies()); }),
+		FCanExecuteAction::CreateSP(this, &FPhysicsAssetEditor::IsNotSimulation));
+
+	ToolkitCommands->MapAction(
+		Commands.CreateBodyWithTaperedCapsule,
+		FExecuteAction::CreateSPLambda(this, [this]() { this->CreateBodiesAndConstraintsForSelectedBones(EAggCollisionShape::TaperedCapsule, this->ShouldCreateConstraintsWhenCreatingBodies()); }),
+		FCanExecuteAction::CreateSP(this, &FPhysicsAssetEditor::IsNotSimulation));
+
+	ToolkitCommands->MapAction(
+		Commands.CreateBodyShouldCreateConstraints,
+		FExecuteAction::CreateSP(this, &FPhysicsAssetEditor::OnToggleCreateConstraintsWhenCreatingBodies),
+		FCanExecuteAction::CreateSP(this, &FPhysicsAssetEditor::IsNotSimulation),
+		FIsActionChecked::CreateSP(this, &FPhysicsAssetEditor::ShouldCreateConstraintsWhenCreatingBodies));
+
 	ToolkitCommands->MapAction(
 		Commands.CreateOrRegenerateBodies,
 		FExecuteAction::CreateSP(this, &FPhysicsAssetEditor::ResetBoneCollision),
@@ -1772,7 +1848,7 @@ void FPhysicsAssetEditor::BuildMenuWidgetBody(FMenuBuilder& InMenuBuilder)
 		InMenuBuilder.BeginSection( "BodyActions", LOCTEXT( "BodyHeader", "Body" ) );
 		InMenuBuilder.AddMenuEntry( Commands.RegenerateBodies );
 		InMenuBuilder.AddSubMenu( LOCTEXT("AddPrimitiveMenu", "Add Primitive"), LOCTEXT("AddPrimitiveMenu_ToolTip", "Add Primitives to this body"),
-			FNewMenuDelegate::CreateStatic( &FillAddShapeMenu ) );
+			FNewMenuDelegate::CreateStatic( &FillAddPrimitiveMenu ) );
 		InMenuBuilder.AddSubMenu( LOCTEXT("CollisionMenu", "Collision"), LOCTEXT("CollisionMenu_ToolTip", "Adjust body/body collision"),
 			FNewMenuDelegate::CreateStatic( &FLocal::FillCollisionMenu ) );	
 		if (bExposeConstraintControls)
@@ -2042,7 +2118,14 @@ TSharedRef<ISkeletonTree> FPhysicsAssetEditor::BuildMenuWidgetNewConstraintForBo
 
 void FPhysicsAssetEditor::BuildMenuWidgetBone(FMenuBuilder& InMenuBuilder)
 {
-	BuildMenuWidgetBone(InMenuBuilder, FPhysicsAssetEditorCommands::Get().CreateBodies);
+	InMenuBuilder.PushCommandList(GetToolkitCommands());
+	InMenuBuilder.BeginSection("BodyActions", LOCTEXT("BodyHeader", "Body"));
+	{
+		InMenuBuilder.AddMenuEntry(FPhysicsAssetEditorCommands::Get().CreateBodies);
+	}
+	InMenuBuilder.EndSection();
+	AddAdvancedMenuWidget(InMenuBuilder);
+	InMenuBuilder.PopCommandList();
 }
 
 bool FPhysicsAssetEditor::ShouldFilterAssetBasedOnSkeleton( const FAssetData& AssetData )
@@ -2437,7 +2520,22 @@ void FPhysicsAssetEditor::OnChangeDefaultMesh(USkeletalMesh* OldPreviewMesh, USk
 	}
 }
 
+void FPhysicsAssetEditor::CreateBodiesAndConstraintsForSelectedBones(const EAggCollisionShape::Type InPrimitiveType, const bool bInShouldCreateConstraints)
+{
+	FPhysAssetCreateParams NewBodyData = GetDefault<UPhysicsAssetGenerationSettings>()->CreateParams;
+
+	NewBodyData.GeomType = ConvertAggCollisionShapeTypeToPhysicsAssetGeomType(InPrimitiveType);
+	NewBodyData.bCreateConstraints = bInShouldCreateConstraints;
+
+	CreateBodiesAndConstraintsForSelectedBones(NewBodyData);
+}
+
 void FPhysicsAssetEditor::ResetBoneCollision()
+{
+	CreateBodiesAndConstraintsForSelectedBones(GetDefault<UPhysicsAssetGenerationSettings>()->CreateParams);
+}
+
+void FPhysicsAssetEditor::CreateBodiesAndConstraintsForSelectedBones(const FPhysAssetCreateParams& NewBodyData)
 {
 	USkeletalMesh* EditorSkelMesh = SharedData->PhysicsAsset->GetPreviewMesh();
 	if(EditorSkelMesh == nullptr)
@@ -2447,8 +2545,6 @@ void FPhysicsAssetEditor::ResetBoneCollision()
 	// Make sure rendering is done - so we are not changing data being used by collision drawing.
 	FlushRenderingCommands();
 
-	const FPhysAssetCreateParams& NewBodyData = GetDefault<UPhysicsAssetGenerationSettings>()->CreateParams;
-
 	if(SharedData->SelectedBodies.Num() > 0)
 	{
 		TArray<int32> SelectedBodyIndices;
@@ -2499,7 +2595,7 @@ void FPhysicsAssetEditor::ResetBoneCollision()
 				const int32 BoneIndex = EditorSkelMesh->GetRefSkeleton().FindBoneIndex(BoneProxy->BoneName);
 				if (BoneIndex != INDEX_NONE)
 				{
-					SharedData->MakeNewBody(BoneIndex);
+					SharedData->MakeNewBody(NewBodyData, BoneIndex);
 				}
 			}
 		}
@@ -3930,6 +4026,12 @@ void FPhysicsAssetEditor::OnSelectAllConstraints()
 
 }
 
+void FPhysicsAssetEditor::OnToggleCreateConstraintsWhenCreatingBodies()
+{
+	SharedData->EditorOptions->bCreateConstraintsWhenCreatingBodiesFromSkeletonTree = !SharedData->EditorOptions->bCreateConstraintsWhenCreatingBodiesFromSkeletonTree;
+	SharedData->EditorOptions->SaveConfig();
+}
+
 void FPhysicsAssetEditor::OnToggleSelectionType(bool bIgnoreUserConstraints)
 {
 	SharedData->ToggleSelectionType(bIgnoreUserConstraints);
@@ -3982,6 +4084,11 @@ void FPhysicsAssetEditor::OnDeselectAll()
 	SharedData->ClearSelectedConstraints();
 }
 
+bool FPhysicsAssetEditor::ShouldCreateConstraintsWhenCreatingBodies() const
+{
+	return SharedData->EditorOptions->bCreateConstraintsWhenCreatingBodiesFromSkeletonTree;
+}
+
 // record if simulating or not, or mode changed or not, or what mode it is in while simulating and what kind of simulation options
 void FPhysicsAssetEditor::OnAddPhatRecord(const FString& Action, bool bRecordSimulate, bool bRecordMode)
 {
@@ -4271,7 +4378,7 @@ void FPhysicsAssetEditor::HandleExtendContextMenu(FMenuBuilder& InMenuBuilder)
 	}
 	else if(SelectedBones.Num() > 0)
 	{
-		BuildMenuWidgetBone(InMenuBuilder, &SelectedBones);
+		BuildMenuWidgetBone(InMenuBuilder, SelectedBones);
 	}
 
 	BuildMenuWidgetSelection(InMenuBuilder);
diff --git a/Engine/Source/Editor/PhysicsAssetEditor/Private/PhysicsAssetEditor.h b/Engine/Source/Editor/PhysicsAssetEditor/Private/PhysicsAssetEditor.h
--- a/Engine/Source/Editor/PhysicsAssetEditor/Private/PhysicsAssetEditor.h
+++ b/Engine/Source/Editor/PhysicsAssetEditor/Private/PhysicsAssetEditor.h
@@ -135,6 +135,9 @@ public:
 
 	TSharedPtr<ISkeletonTree> GetSkeletonTree() const { return SkeletonTree.ToSharedRef(); }
 
+	void CreateBodiesAndConstraintsForSelectedBones(const EAggCollisionShape::Type InPrimitiveType, const bool bInShouldCreateConstraints);
+	void CreateBodiesAndConstraintsForSelectedBones(const FPhysAssetCreateParams& InBodyData);
+
 	/** Reset bone collision for selected or regenerate all bodies if no bodies are selected */
 	void ResetBoneCollision();
 
@@ -319,6 +322,7 @@ private:
 	void OnSelectBodies(EPhysicsType PhysicsType = EPhysicsType::PhysType_Simulated);
 	void OnSelectShapes(const ECollisionEnabled::Type CollisionEnabled);
 	void OnSelectAllConstraints();
+	void OnToggleCreateConstraintsWhenCreatingBodies();
 	void OnToggleSelectionType(bool bIgnoreUserConstraints);
 	void OnToggleShowSelected();
 	void OnShowSelected();
@@ -331,15 +335,13 @@ private:
 	void OnDeselectAll();
 	bool LoadMLLevelSetFromDataTable(UDataTable* DataTableMLLevelSetModelBoneBinningInfo, TArray<FString>& ErrorMessages);
 	void ImportMLLevelSetFromDataTable();
-
+	bool ShouldCreateConstraintsWhenCreatingBodies() const;
+	
 	FText GetRepeatLastSimulationToolTip() const;
 	FSlateIcon GetRepeatLastSimulationIcon() const;
 
-	/** Add a specified skeleton tree context menu entriy specific to bones. */
-	void BuildMenuWidgetBone(FMenuBuilder& InMenuBuilder, const TSharedPtr< const FUICommandInfo > InCommand);
-
 	/** Add skeleton tree context menu entries specific to the currently selected set of bones. */
-	void BuildMenuWidgetBone(FMenuBuilder& InMenuBuilder, const TArray<TSharedPtr<ISkeletonTreeItem>>* const SelectedBones);
+	void BuildMenuWidgetBone(FMenuBuilder& InMenuBuilder, const TArray<TSharedPtr<ISkeletonTreeItem>>& SelectedBones);
 
 	/** Handle initial preview scene setup */
 	void HandlePreviewSceneCreated(const TSharedRef<IPersonaPreviewScene>& InPersonaPreviewScene);
diff --git a/Engine/Source/Editor/PhysicsAssetEditor/Private/PhysicsAssetEditorActions.cpp b/Engine/Source/Editor/PhysicsAssetEditor/Private/PhysicsAssetEditorActions.cpp
--- a/Engine/Source/Editor/PhysicsAssetEditor/Private/PhysicsAssetEditorActions.cpp
+++ b/Engine/Source/Editor/PhysicsAssetEditor/Private/PhysicsAssetEditorActions.cpp
@@ -13,8 +13,12 @@ void FPhysicsAssetEditorCommands::RegisterCommands()
 {
 	UI_COMMAND(RegenerateBodies, "Regenerate Bodies", "Regenerates the selected bodies using the current generation settings (see the Tools tab)", EUserInterfaceActionType::Button, FInputChord());
 	UI_COMMAND(CreateBodies, "Create Bodies/Constraints", "Creates bodies and/or constraints for the selected bones using the current generation settings (see the Tools tab)", EUserInterfaceActionType::Button, FInputChord());
+	UI_COMMAND(CreateBodyWithBox, "Create Bodies with Box", "Create bodies with the Box primitive type (and optionally create constraints) for the selected bones using current generation settings (see the Tools tab)", EUserInterfaceActionType::Button, FInputChord());
+	UI_COMMAND(CreateBodyWithSphere, "Create Bodies with Sphere", "Create bodies with the Sphere primitive type (and optionally create constraints) for the selected bones using current generation settings (see the Tools tab)", EUserInterfaceActionType::Button, FInputChord());
+	UI_COMMAND(CreateBodyWithSphyl, "Create Bodies with Capsule", "Create bodies with the Capsule primitive type (and optionally create constraints) for the selected bones using current generation settings (see the Tools tab)", EUserInterfaceActionType::Button, FInputChord());
+	UI_COMMAND(CreateBodyWithTaperedCapsule, "Create Bodies with Tapered Capsule (Clothing Only)", "Create bodies with the Tapered Capsule primitive type (and optionally create constraints) for the selected bones using current generation settings (see the Tools tab)", EUserInterfaceActionType::Button, FInputChord());
+	UI_COMMAND(CreateBodyShouldCreateConstraints, "Create Constraints", "Create constraints between any newly created bodies and the bodies on adjacent bones for each of the selected bones", EUserInterfaceActionType::ToggleButton, FInputChord());
 	UI_COMMAND(CreateOrRegenerateBodies, "Create/Regenerate Bodies", "Creates bodies and/or constraints (if none exist) or Regenerates bodies (if they do exist) for each of the selected bones using the current generation settings (see the Tools tab)", EUserInterfaceActionType::Button, FInputChord());
-
 	UI_COMMAND(ApplyPhysicalMaterial, "Physical Material", "Apply a Physical Material to All Bodies", EUserInterfaceActionType::Button, FInputChord());
 	UI_COMMAND(CopyBodies, "Copy Selected Bodies/Constraints To Clipboard", "Copy selected bodies/constraints to clipboard", EUserInterfaceActionType::Button, FInputChord(EModifierKey::Control | EModifierKey::Shift, EKeys::C));
 	UI_COMMAND(PasteBodies, "Paste Bodies/Constraints From Clipboard", "Paste bodies/constraints from clipboard", EUserInterfaceActionType::Button, FInputChord(EModifierKey::Control | EModifierKey::Shift, EKeys::V));
diff --git a/Engine/Source/Editor/PhysicsAssetEditor/Private/PhysicsAssetEditorActions.h b/Engine/Source/Editor/PhysicsAssetEditor/Private/PhysicsAssetEditorActions.h
--- a/Engine/Source/Editor/PhysicsAssetEditor/Private/PhysicsAssetEditorActions.h
+++ b/Engine/Source/Editor/PhysicsAssetEditor/Private/PhysicsAssetEditorActions.h
@@ -28,6 +28,11 @@ public:
 	
 	TSharedPtr<FUICommandInfo> RegenerateBodies;
 	TSharedPtr<FUICommandInfo> CreateBodies;
+	TSharedPtr<FUICommandInfo> CreateBodyWithSphere;
+	TSharedPtr<FUICommandInfo> CreateBodyWithSphyl;
+	TSharedPtr<FUICommandInfo> CreateBodyWithBox;
+	TSharedPtr<FUICommandInfo> CreateBodyWithTaperedCapsule;
+	TSharedPtr<FUICommandInfo> CreateBodyShouldCreateConstraints;
 	TSharedPtr<FUICommandInfo> CreateOrRegenerateBodies;
 	TSharedPtr<FUICommandInfo> ApplyPhysicalMaterial;
 	TSharedPtr<FUICommandInfo> Snap;
diff --git a/Engine/Source/Editor/PhysicsAssetEditor/Private/PhysicsAssetEditorSharedData.cpp b/Engine/Source/Editor/PhysicsAssetEditor/Private/PhysicsAssetEditorSharedData.cpp
--- a/Engine/Source/Editor/PhysicsAssetEditor/Private/PhysicsAssetEditorSharedData.cpp
+++ b/Engine/Source/Editor/PhysicsAssetEditor/Private/PhysicsAssetEditorSharedData.cpp
@@ -1602,7 +1602,7 @@ bool FPhysicsAssetEditorSharedData::GetPrimitiveContributeToMass() const
 	return false;
 }
 
-static EAggCollisionShape::Type ConvertPhysicsAssetGeomTypeToAggCollisionShapeType(EPhysAssetFitGeomType PhysicsAssetGeomType)
+EAggCollisionShape::Type ConvertPhysicsAssetGeomTypeToAggCollisionShapeType(EPhysAssetFitGeomType PhysicsAssetGeomType)
 {
 	switch (PhysicsAssetGeomType)
 	{
@@ -1617,6 +1617,21 @@ static EAggCollisionShape::Type ConvertPhysicsAssetGeomTypeToAggCollisionShapeTy
 	}
 }
 
+EPhysAssetFitGeomType ConvertAggCollisionShapeTypeToPhysicsAssetGeomType(const EAggCollisionShape::Type AggCollisionShapeType)
+{
+	switch (AggCollisionShapeType)
+	{
+	case EAggCollisionShape::Type::Box:					return EPhysAssetFitGeomType::EFG_Box;
+	case EAggCollisionShape::Type::Sphyl:				return EPhysAssetFitGeomType::EFG_Sphyl;
+	case EAggCollisionShape::Type::Sphere:				return EPhysAssetFitGeomType::EFG_Sphere;
+	case EAggCollisionShape::Type::TaperedCapsule:		return EPhysAssetFitGeomType::EFG_TaperedCapsule;
+	case EAggCollisionShape::Type::Convex:				return EPhysAssetFitGeomType::EFG_SingleConvexHull;
+	case EAggCollisionShape::Type::LevelSet:			return EPhysAssetFitGeomType::EFG_LevelSet;
+	default:											return EPhysAssetFitGeomType(INDEX_NONE);
+	}
+}
+
+
 void FPhysicsAssetEditorSharedData::AutoNameAllPrimitives(int32 BodyIndex, EPhysAssetFitGeomType PrimitiveType)
 {
 	AutoNameAllPrimitives(BodyIndex, ConvertPhysicsAssetGeomTypeToAggCollisionShapeType(PrimitiveType));
@@ -2270,6 +2285,11 @@ void FPhysicsAssetEditorSharedData::InitConstraintSetup(UPhysicsConstraintTempla
 }
 
 void FPhysicsAssetEditorSharedData::MakeNewBody(int32 NewBoneIndex, bool bAutoSelect)
+{
+	MakeNewBody(GetDefault<UPhysicsAssetGenerationSettings>()->CreateParams, NewBoneIndex, bAutoSelect);
+}
+
+void FPhysicsAssetEditorSharedData::MakeNewBody(const FPhysAssetCreateParams& NewBodyData, const int32 NewBoneIndex, const bool bAutoSelect)
 {
 	USkeletalMesh* EditorSkelMesh = PhysicsAsset->GetPreviewMesh();
 	if(EditorSkelMesh == nullptr)
@@ -2290,8 +2310,6 @@ void FPhysicsAssetEditorSharedData::MakeNewBody(int32 NewBoneIndex, bool bAutoSe
 	// Find body that currently controls this bone.
 	int32 ParentBodyIndex = PhysicsAsset->FindControllingBodyIndex(EditorSkelMesh, NewBoneIndex);
 
-	const FPhysAssetCreateParams& NewBodyData = GetDefault<UPhysicsAssetGenerationSettings>()->CreateParams;
-
 	// Create the physics body.
 	NewBodyIndex = FPhysicsAssetUtils::CreateNewBody(PhysicsAsset, NewBoneName, NewBodyData);
 	UBodySetup* BodySetup = PhysicsAsset->SkeletalBodySetups[ NewBodyIndex ];
@@ -2318,8 +2336,8 @@ void FPhysicsAssetEditorSharedData::MakeNewBody(int32 NewBoneIndex, bool bAutoSe
 
 	// name the new created primitives
 	AutoNameAllPrimitives(NewBodyIndex, NewBodyData.GeomType);
-
-	const bool bCanCreateConstraints = FPhysicsAssetUtils::CanCreateConstraints();
+	
+	const bool bCreateConstraints = NewBodyData.bCreateConstraints && FPhysicsAssetUtils::CanCreateConstraints();
 
 	// Check if the bone of the new body has any physical children bones
 	for (int32 i = 0; i < EditorSkelMesh->GetRefSkeleton().GetRawBoneNum(); ++i)
@@ -2347,7 +2365,7 @@ void FPhysicsAssetEditorSharedData::MakeNewBody(int32 NewBoneIndex, bool bAutoSe
 				// We are currently doing the latter...
 				if (ConstraintIndex == INDEX_NONE)
 				{
-					if (bCanCreateConstraints)
+					if (bCreateConstraints)
 					{
 						ConstraintIndex = FPhysicsAssetUtils::CreateNewConstraint(PhysicsAsset, ChildBody->BoneName);
 						check(ConstraintIndex != INDEX_NONE);
@@ -2394,7 +2412,7 @@ void FPhysicsAssetEditorSharedData::MakeNewBody(int32 NewBoneIndex, bool bAutoSe
 	}
 
 	// If we have a physics parent, create a joint to it.
-	if (ParentBodyIndex != INDEX_NONE && bCanCreateConstraints)
+	if (ParentBodyIndex != INDEX_NONE && bCreateConstraints)
 	{
 		const int32 NewConstraintIndex = FPhysicsAssetUtils::CreateNewConstraint(PhysicsAsset, NewBoneName);
 		UPhysicsConstraintTemplate* ConstraintSetup = PhysicsAsset->ConstraintSetup[ NewConstraintIndex ];
diff --git a/Engine/Source/Editor/PhysicsAssetEditor/Private/PhysicsAssetEditorSharedData.h b/Engine/Source/Editor/PhysicsAssetEditor/Private/PhysicsAssetEditorSharedData.h
--- a/Engine/Source/Editor/PhysicsAssetEditor/Private/PhysicsAssetEditorSharedData.h
+++ b/Engine/Source/Editor/PhysicsAssetEditor/Private/PhysicsAssetEditorSharedData.h
@@ -204,6 +204,7 @@ public:
 	void DeleteBody(int32 DelBodyIndex, bool bRefreshComponent=true);
 	void RefreshPhysicsAssetChange(const UPhysicsAsset* InPhysAsset, bool bFullClothRefresh = true);
 	void MakeNewBody(int32 NewBoneIndex, bool bAutoSelect = true);
+	void MakeNewBody(const FPhysAssetCreateParams& NewBodyData, const int32 NewBoneIndex, const bool bAutoSelect = true);
 	void MakeNewConstraints(int32 ParentBodyIndex, const TArray<int32>& ChildBodyIndices);
 	void MakeNewConstraint(int32 ParentBodyIndex, int32 ChildBodyIndex);
 	void CopySelectedBodiesAndConstraintsToClipboard(int32& OutNumCopiedBodies, int32& OutNumCopiedConstraints);
@@ -438,3 +439,6 @@ public:
 	FVector LastClickHitNormal;
 	bool bLastClickHit;
 };
+
+EAggCollisionShape::Type ConvertPhysicsAssetGeomTypeToAggCollisionShapeType(EPhysAssetFitGeomType PhysicsAssetGeomType);
+EPhysAssetFitGeomType ConvertAggCollisionShapeTypeToPhysicsAssetGeomType(const EAggCollisionShape::Type AggCollisionShapeType);
\ No newline at end of file
diff --git a/Engine/Source/Editor/UnrealEd/Classes/Preferences/PhysicsAssetEditorOptions.h b/Engine/Source/Editor/UnrealEd/Classes/Preferences/PhysicsAssetEditorOptions.h
--- a/Engine/Source/Editor/UnrealEd/Classes/Preferences/PhysicsAssetEditorOptions.h
+++ b/Engine/Source/Editor/UnrealEd/Classes/Preferences/PhysicsAssetEditorOptions.h
@@ -187,6 +187,10 @@ class UPhysicsAssetEditorOptions : public UObject
 	UPROPERTY(config)
 	bool bHideBodyMass;
 
+	/** When set, constraints will be automatically created when physics bodies are created through the context menu in the skeleton tree. */
+	UPROPERTY(config)
+	bool bCreateConstraintsWhenCreatingBodiesFromSkeletonTree;
+
 	/** When set, cloth will reset each time simulation is toggled */
 	UPROPERTY(EditAnywhere, config, Category=Clothing)
 	bool bResetClothWhenSimulating;
diff --git a/Engine/Source/Editor/UnrealEd/Private/PreferenceStubs.cpp b/Engine/Source/Editor/UnrealEd/Private/PreferenceStubs.cpp
--- a/Engine/Source/Editor/UnrealEd/Private/PreferenceStubs.cpp
+++ b/Engine/Source/Editor/UnrealEd/Private/PreferenceStubs.cpp
@@ -63,6 +63,7 @@ UPhysicsAssetEditorOptions::UPhysicsAssetEditorOptions(const FObjectInitializer&
 	bHideSimulatedBodies = false;
 	bHideKinematicBodies = false;
 	bHideBodyMass = true;
+	bCreateConstraintsWhenCreatingBodiesFromSkeletonTree = true;
 	bResetClothWhenSimulating = false;
 }
 
