[UE] PhAT cross constraint icons and visibility filters in skeleton tree view
- Introduce a new icon to represent cross constraints in the Physics Asset Editor
- Update constraint existing and body icons in the Physics Asset Editor with new images
- Add separate skeleton tree filter options for cross constraints, parent-child (normal) constraints and an overarching constraints option that hides both kinds.
- Gray out cross and parent-child constraint filter options when the overarching constraints filter option is unchecked
- Gray out simulated and kinematic body filter options when the overarching bodies filter option is unchecked

diff --git a/Engine/Source/Editor/EditorStyle/Private/SlateEditorStyle.cpp b/Engine/Source/Editor/EditorStyle/Private/SlateEditorStyle.cpp
--- a/Engine/Source/Editor/EditorStyle/Private/SlateEditorStyle.cpp
+++ b/Engine/Source/Editor/EditorStyle/Private/SlateEditorStyle.cpp
@@ -7645,7 +7645,7 @@ void FSlateEditorStyle::FStyle::SetupToolkitStyles()
 		Set("PhysicsAssetEditor.AddConstraintToCurrentConstraintProfile", new IMAGE_BRUSH("PhysicsAssetEditor/icon_PHatMode_Joint_40x", Icon20x20));
 		Set("PhysicsAssetEditor.RemoveConstraintFromCurrentConstraintProfile", new IMAGE_BRUSH("PhysicsAssetEditor/icon_PhAT_DeleteConstraint_40x", Icon20x20));
 
-		Set("PhysicsAssetEditor.Tree.Body", new IMAGE_BRUSH("PhysicsAssetEditor/Body_16x", Icon16x16));
+		Set("PhysicsAssetEditor.Tree.Body", new IMAGE_BRUSH("PhysicsAssetEditor/BoneGroup_16x", Icon16x16));
 		Set("PhysicsAssetEditor.Tree.KinematicBody", new IMAGE_BRUSH("PhysicsAssetEditor/KinematicBody_16x", Icon16x16));
 		Set("PhysicsAssetEditor.Tree.EmptyBody", new IMAGE_BRUSH("PhysicsAssetEditor/EmptyBody_16x", Icon16x16));
 		Set("PhysicsAssetEditor.Tree.Bone", new IMAGE_BRUSH("PhysicsAssetEditor/Bone_16x", Icon16x16));
@@ -7654,7 +7654,8 @@ void FSlateEditorStyle::FStyle::SetupToolkitStyles()
 		Set("PhysicsAssetEditor.Tree.Box", new IMAGE_BRUSH("PhysicsAssetEditor/Box_16x", Icon16x16));
 		Set("PhysicsAssetEditor.Tree.Convex", new IMAGE_BRUSH("PhysicsAssetEditor/Convex_16x", Icon16x16));
 		Set("PhysicsAssetEditor.Tree.TaperedCapsule", new IMAGE_BRUSH("PhysicsAssetEditor/TaperedCapsule_16x", Icon16x16));
-		Set("PhysicsAssetEditor.Tree.Constraint", new IMAGE_BRUSH("PhysicsAssetEditor/Constraint_16x", Icon16x16));
+		Set("PhysicsAssetEditor.Tree.Constraint", new IMAGE_BRUSH("PhysicsAssetEditor/BoneConstraint_16x", Icon16x16));
+		Set("PhysicsAssetEditor.Tree.CrossConstraint", new IMAGE_BRUSH("PhysicsAssetEditor/BoneCrossConstraint_16x", Icon16x16));
 
 		Set("PhysicsAssetEditor.Tree.Font", DEFAULT_FONT("Regular", 10));
 
diff --git a/Engine/Source/Editor/EditorStyle/Private/StarshipStyle.cpp b/Engine/Source/Editor/EditorStyle/Private/StarshipStyle.cpp
--- a/Engine/Source/Editor/EditorStyle/Private/StarshipStyle.cpp
+++ b/Engine/Source/Editor/EditorStyle/Private/StarshipStyle.cpp
@@ -7518,7 +7518,7 @@ void FStarshipEditorStyle::FStyle::SetupToolkitStyles()
 		Set("PhysicsAssetEditor.AddConstraintToCurrentConstraintProfile", new IMAGE_BRUSH("PhysicsAssetEditor/icon_PHatMode_Joint_40x", Icon20x20));
 		Set("PhysicsAssetEditor.RemoveConstraintFromCurrentConstraintProfile", new IMAGE_BRUSH("PhysicsAssetEditor/icon_PhAT_DeleteConstraint_40x", Icon20x20));
 
-		Set("PhysicsAssetEditor.Tree.Body", new IMAGE_BRUSH("PhysicsAssetEditor/Body_16x", Icon16x16));
+		Set("PhysicsAssetEditor.Tree.Body", new IMAGE_BRUSH("PhysicsAssetEditor/BoneGroup_16x", Icon16x16));
 		Set("PhysicsAssetEditor.Tree.KinematicBody", new IMAGE_BRUSH("PhysicsAssetEditor/KinematicBody_16x", Icon16x16));
 		Set("PhysicsAssetEditor.Tree.EmptyBody", new IMAGE_BRUSH("PhysicsAssetEditor/EmptyBody_16x", Icon16x16));
 		Set("PhysicsAssetEditor.Tree.Bone", new IMAGE_BRUSH("PhysicsAssetEditor/Bone_16x", Icon16x16));
@@ -7527,7 +7527,8 @@ void FStarshipEditorStyle::FStyle::SetupToolkitStyles()
 		Set("PhysicsAssetEditor.Tree.Box", new IMAGE_BRUSH("PhysicsAssetEditor/Box_16x", Icon16x16));
 		Set("PhysicsAssetEditor.Tree.Convex", new IMAGE_BRUSH("PhysicsAssetEditor/Convex_16x", Icon16x16));
 		Set("PhysicsAssetEditor.Tree.TaperedCapsule", new IMAGE_BRUSH("PhysicsAssetEditor/TaperedCapsule_16x", Icon16x16));
-		Set("PhysicsAssetEditor.Tree.Constraint", new IMAGE_BRUSH("PhysicsAssetEditor/Constraint_16x", Icon16x16));
+		Set("PhysicsAssetEditor.Tree.Constraint", new IMAGE_BRUSH("PhysicsAssetEditor/BoneConstraint_16x", Icon16x16));
+		Set("PhysicsAssetEditor.Tree.CrossConstraint", new IMAGE_BRUSH("PhysicsAssetEditor/BoneCrossConstraint_16x", Icon16x16));
 
 		Set("PhysicsAssetEditor.BoneAssign", new IMAGE_BRUSH_SVG("Starship/Persona/BoneAssign", Icon20x20));
 		Set("PhysicsAssetEditor.BoneUnassign", new IMAGE_BRUSH_SVG("Starship/Persona/BoneUnassign", Icon20x20));
diff --git a/Engine/Source/Editor/PhysicsAssetEditor/Private/PhysicsAssetEditor.cpp b/Engine/Source/Editor/PhysicsAssetEditor/Private/PhysicsAssetEditor.cpp
--- a/Engine/Source/Editor/PhysicsAssetEditor/Private/PhysicsAssetEditor.cpp
+++ b/Engine/Source/Editor/PhysicsAssetEditor/Private/PhysicsAssetEditor.cpp
@@ -1713,14 +1713,14 @@ void FPhysicsAssetEditor::BindCommands()
 	SkeletonTreeCommandList->MapAction(
 		Commands.ShowSimulatedBodies,
 		FExecuteAction::CreateSP(this, &FPhysicsAssetEditor::HandleToggleShowSimulatedBodies),
-		FCanExecuteAction(),
+		FCanExecuteAction::CreateSP(this, &FPhysicsAssetEditor::IsShowBodiesChecked),
 		FGetActionCheckState::CreateSP(this, &FPhysicsAssetEditor::GetShowSimulatedBodiesChecked)
 	);
 
 	SkeletonTreeCommandList->MapAction(
 		Commands.ShowKinematicBodies,
 		FExecuteAction::CreateSP(this, &FPhysicsAssetEditor::HandleToggleShowKinematicBodies),
-		FCanExecuteAction(),
+		FCanExecuteAction::CreateSP(this, &FPhysicsAssetEditor::IsShowBodiesChecked),
 		FGetActionCheckState::CreateSP(this, &FPhysicsAssetEditor::GetShowKinematicBodiesChecked)
 	);
 
@@ -1731,6 +1731,20 @@ void FPhysicsAssetEditor::BindCommands()
 		FGetActionCheckState::CreateSP(this, &FPhysicsAssetEditor::GetShowConstraintsChecked)
 	);
 
+	SkeletonTreeCommandList->MapAction(
+		Commands.ShowCrossConstraints,
+		FExecuteAction::CreateSP(this, &FPhysicsAssetEditor::HandleToggleShowCrossConstraints),
+		FCanExecuteAction::CreateSP(this, &FPhysicsAssetEditor::IsShowConstraintsChecked),
+		FGetActionCheckState::CreateSP(this, &FPhysicsAssetEditor::GetShowCrossConstraintsChecked)
+	);
+
+	SkeletonTreeCommandList->MapAction(
+		Commands.ShowParentChildConstraints,
+		FExecuteAction::CreateSP(this, &FPhysicsAssetEditor::HandleToggleShowParentChildConstraints),
+		FCanExecuteAction::CreateSP(this, &FPhysicsAssetEditor::IsShowConstraintsChecked),
+		FGetActionCheckState::CreateSP(this, &FPhysicsAssetEditor::GetShowParentChildConstraintsChecked)
+	);
+
 	SkeletonTreeCommandList->MapAction(
 		Commands.ShowConstraintsOnParentBodies,
 		FExecuteAction::CreateSP(this, &FPhysicsAssetEditor::HandleToggleShowConstraintsOnParentBodies),
@@ -2051,6 +2065,8 @@ TSharedRef<ISkeletonTree> FPhysicsAssetEditor::BuildMenuWidgetNewConstraintForBo
 	Builder->bShowSimulatedBodies = true;
 	Builder->bShowKinematicBodies = true;
 	Builder->bShowConstraints = false;
+	Builder->bShowCrossConstraints = true;
+	Builder->bShowParentChildConstraints = true;
 	Builder->bShowPrimitives = false;
 
 	FSkeletonTreeArgs SkeletonTreeArgs;
@@ -4403,6 +4419,8 @@ void FPhysicsAssetEditor::HandleExtendFilterMenu(FMenuBuilder& InMenuBuilder)
 		if (bExposeConstraintControls)
 		{
 			InMenuBuilder.AddMenuEntry(Commands.ShowConstraints);
+			InMenuBuilder.AddMenuEntry(Commands.ShowParentChildConstraints);
+			InMenuBuilder.AddMenuEntry(Commands.ShowCrossConstraints);
 		}
 		InMenuBuilder.AddMenuEntry(Commands.ShowPrimitives);
 		if (bExposeConstraintControls)
@@ -4439,6 +4457,18 @@ void FPhysicsAssetEditor::HandleToggleShowConstraints()
 	RefreshFilter();
 }
 
+void FPhysicsAssetEditor::HandleToggleShowCrossConstraints()
+{
+	SkeletonTreeBuilder->bShowCrossConstraints = !SkeletonTreeBuilder->bShowCrossConstraints;
+	RefreshFilter();
+}
+
+void FPhysicsAssetEditor::HandleToggleShowParentChildConstraints()
+{
+	SkeletonTreeBuilder->bShowParentChildConstraints = !SkeletonTreeBuilder->bShowParentChildConstraints;
+	RefreshFilter();
+}
+
 void FPhysicsAssetEditor::HandleToggleShowConstraintsOnParentBodies()
 {
 	SkeletonTreeBuilder->bShowConstraintsOnParentBodies = !SkeletonTreeBuilder->bShowConstraintsOnParentBodies;
@@ -4466,9 +4496,19 @@ ECheckBoxState FPhysicsAssetEditor::GetShowKinematicBodiesChecked() const
 	return SkeletonTreeBuilder->bShowKinematicBodies ? ECheckBoxState::Checked : ECheckBoxState::Unchecked;
 }
 
-ECheckBoxState FPhysicsAssetEditor::GetShowConstraintsChecked() const
+ECheckBoxState FPhysicsAssetEditor::GetShowCrossConstraintsChecked() const
 {
-	return SkeletonTreeBuilder->bShowConstraints ? ECheckBoxState::Checked : ECheckBoxState::Unchecked;
+	return SkeletonTreeBuilder->bShowCrossConstraints ? ECheckBoxState::Checked : ECheckBoxState::Unchecked;
+}
+
+ECheckBoxState FPhysicsAssetEditor::GetShowParentChildConstraintsChecked() const
+{
+	return SkeletonTreeBuilder->bShowParentChildConstraints ? ECheckBoxState::Checked : ECheckBoxState::Unchecked;
+}
+
+bool FPhysicsAssetEditor::IsShowBodiesChecked() const
+{
+	return SkeletonTreeBuilder->bShowBodies;
 }
 
 bool FPhysicsAssetEditor::IsShowConstraintsChecked() const
@@ -4481,6 +4521,11 @@ ECheckBoxState FPhysicsAssetEditor::GetShowConstraintsOnParentBodiesChecked() co
 	return SkeletonTreeBuilder->bShowConstraintsOnParentBodies ? ECheckBoxState::Checked : ECheckBoxState::Unchecked;
 }
 
+ECheckBoxState FPhysicsAssetEditor::GetShowConstraintsChecked() const
+{
+	return SkeletonTreeBuilder->bShowConstraints ? ECheckBoxState::Checked : ECheckBoxState::Unchecked;
+}
+
 ECheckBoxState FPhysicsAssetEditor::GetShowPrimitivesChecked() const
 {
 	return SkeletonTreeBuilder->bShowPrimitives ? ECheckBoxState::Checked : ECheckBoxState::Unchecked;
diff --git a/Engine/Source/Editor/PhysicsAssetEditor/Private/PhysicsAssetEditor.h b/Engine/Source/Editor/PhysicsAssetEditor/Private/PhysicsAssetEditor.h
--- a/Engine/Source/Editor/PhysicsAssetEditor/Private/PhysicsAssetEditor.h
+++ b/Engine/Source/Editor/PhysicsAssetEditor/Private/PhysicsAssetEditor.h
@@ -361,13 +361,18 @@ private:
 	void HandleToggleShowKinematicBodies();
 	void HandleToggleShowConstraints();
 	void HandleToggleShowConstraintsOnParentBodies();
+	void HandleToggleShowCrossConstraints();
+	void HandleToggleShowParentChildConstraints();
 	void HandleToggleShowPrimitives();
 	ECheckBoxState GetShowBodiesChecked() const;
 	ECheckBoxState GetShowSimulatedBodiesChecked() const;
 	ECheckBoxState GetShowKinematicBodiesChecked() const;
 	ECheckBoxState GetShowConstraintsChecked() const;
 	ECheckBoxState GetShowConstraintsOnParentBodiesChecked() const;
+	ECheckBoxState GetShowCrossConstraintsChecked() const;
+	ECheckBoxState GetShowParentChildConstraintsChecked() const;
 	ECheckBoxState GetShowPrimitivesChecked() const;
+	bool IsShowBodiesChecked() const;
 	bool IsShowConstraintsChecked() const;
 
 	/** Customize the filter label */
diff --git a/Engine/Source/Editor/PhysicsAssetEditor/Private/PhysicsAssetEditorActions.cpp b/Engine/Source/Editor/PhysicsAssetEditor/Private/PhysicsAssetEditorActions.cpp
--- a/Engine/Source/Editor/PhysicsAssetEditor/Private/PhysicsAssetEditorActions.cpp
+++ b/Engine/Source/Editor/PhysicsAssetEditor/Private/PhysicsAssetEditorActions.cpp
@@ -148,6 +148,8 @@ void FPhysicsAssetEditorCommands::RegisterCommands()
 	UI_COMMAND(ShowKinematicBodies, "Show Kinematic Bodies", "Display kinematic bodies in the tree view. Bodies are a collection of primitive shapes used for collision.", EUserInterfaceActionType::ToggleButton, FInputChord());
 	UI_COMMAND(ShowConstraints, "Show Constraints", "Display constraints in the tree view. Constraints are used to control how bodies can move in relation to one another.", EUserInterfaceActionType::ToggleButton, FInputChord());
 	UI_COMMAND(ShowConstraintsOnParentBodies, "Show Constraints on Parent Bodies", "When showing constraints, display them on bothe the parent and the child body.", EUserInterfaceActionType::ToggleButton, FInputChord());
+	UI_COMMAND(ShowCrossConstraints, "Show Cross Constraints", "Display cross constraints in the tree view. Constraints are used to control how bodies can move in relation to one another. Cross constraints connect bodies that are attached to bones that are not adjacent (parent/child) in the skeleton.", EUserInterfaceActionType::ToggleButton, FInputChord());
+	UI_COMMAND(ShowParentChildConstraints, "Show Parent-Child Constraints", "Display Parent-Child constraints in the tree view. Constraints are used to control how bodies can move in relation to one another. Parent-Child constraints connect bodies that are attached to bones that are adjacent (parent/child) in the skeleton.", EUserInterfaceActionType::ToggleButton, FInputChord());
 	UI_COMMAND(ShowPrimitives, "Show Primitives", "Display primitive shapes (sphere, boxes, capsules etc.) in the tree view", EUserInterfaceActionType::ToggleButton, FInputChord());
 
 	/** As two commands cannot have the same key; this command wraps both 
diff --git a/Engine/Source/Editor/PhysicsAssetEditor/Private/PhysicsAssetEditorActions.h b/Engine/Source/Editor/PhysicsAssetEditor/Private/PhysicsAssetEditorActions.h
--- a/Engine/Source/Editor/PhysicsAssetEditor/Private/PhysicsAssetEditorActions.h
+++ b/Engine/Source/Editor/PhysicsAssetEditor/Private/PhysicsAssetEditorActions.h
@@ -161,6 +161,8 @@ public:
 	TSharedPtr<FUICommandInfo> ShowKinematicBodies;
 	TSharedPtr<FUICommandInfo> ShowConstraints;
 	TSharedPtr<FUICommandInfo> ShowConstraintsOnParentBodies;
+	TSharedPtr<FUICommandInfo> ShowCrossConstraints;
+	TSharedPtr<FUICommandInfo> ShowParentChildConstraints;
 	TSharedPtr<FUICommandInfo> ShowPrimitives;
 
 	/** Hotkey only commands */
diff --git a/Engine/Source/Editor/PhysicsAssetEditor/Private/PhysicsAssetEditorSkeletonTreeBuilder.cpp b/Engine/Source/Editor/PhysicsAssetEditor/Private/PhysicsAssetEditorSkeletonTreeBuilder.cpp
--- a/Engine/Source/Editor/PhysicsAssetEditor/Private/PhysicsAssetEditorSkeletonTreeBuilder.cpp
+++ b/Engine/Source/Editor/PhysicsAssetEditor/Private/PhysicsAssetEditorSkeletonTreeBuilder.cpp
@@ -5,6 +5,7 @@
 #include "SkeletonTreePhysicsShapeItem.h"
 #include "SkeletonTreePhysicsConstraintItem.h"
 #include "PhysicsEngine/PhysicsAsset.h"
+#include "PhysicsEngine/PhysicsConstraintTemplate.h"
 #include "PhysicsEngine/SkeletalBodySetup.h"
 #include "IPersonaPreviewScene.h"
 #include "Animation/DebugSkelMeshComponent.h"
@@ -19,6 +20,8 @@ FPhysicsAssetEditorSkeletonTreeBuilder::FPhysicsAssetEditorSkeletonTreeBuilder(U
 	, bShowSimulatedBodies(true)
 	, bShowConstraints(false)
 	, bShowConstraintsOnParentBodies(true)
+	, bShowCrossConstraints(true)
+	, bShowParentChildConstraints(true)
 	, bShowPrimitives(false)
 	, PhysicsAsset(InPhysicsAsset)
 {
@@ -89,13 +92,22 @@ ESkeletonTreeFilterResult FPhysicsAssetEditorSkeletonTreeBuilder::FilterItem(con
 			{
 				Result = ESkeletonTreeFilterResult::Hidden;
 			} 
-			else
+			else if(TSharedPtr<FSkeletonTreePhysicsConstraintItem> SkeletonTreePhysicsConstraintItem = StaticCastSharedPtr<FSkeletonTreePhysicsConstraintItem>(InItem))
 			{
-				TSharedPtr<FSkeletonTreePhysicsConstraintItem> SkeletonTreePhysicsConstraintItem = StaticCastSharedPtr<FSkeletonTreePhysicsConstraintItem>(InItem);
 				if (!bShowConstraintsOnParentBodies && SkeletonTreePhysicsConstraintItem->IsConstraintOnParentBody())
 				{
 					Result = ESkeletonTreeFilterResult::Hidden;
 				}
+
+				if (!bShowCrossConstraints && SkeletonTreePhysicsConstraintItem->IsCrossConstraint())
+				{
+					Result = ESkeletonTreeFilterResult::Hidden;
+				}
+
+				if (!bShowParentChildConstraints && !SkeletonTreePhysicsConstraintItem->IsCrossConstraint())
+				{
+					Result = ESkeletonTreeFilterResult::Hidden;
+				}
 			}
 		}
 		else if(InItem->IsOfType<FSkeletonTreePhysicsShapeItem>())
@@ -184,8 +196,12 @@ void FPhysicsAssetEditorSkeletonTreeBuilder::AddBodies(FSkeletonTreeBuilderOutpu
 
 							if (bJointMatches || bUserConstraintMatches)
 							{
+								const int32 ChildBoneIndex = RefSkeleton.FindBoneIndex(ConstraintInstance.GetChildBoneName());
+								const int32 ParentBoneIndex = RefSkeleton.FindBoneIndex(ConstraintInstance.GetParentBoneName());
+								const bool bIsCrossConstraint = RefSkeleton.IsValidIndex(ChildBoneIndex) && (ParentBoneIndex != RefSkeleton.GetParentIndex(ChildBoneIndex));
 								const bool bIsConstraintOnParentBone = !(ConstraintInstance.JointName == BoneName || ConstraintInstance.ConstraintBone1 == BoneName);
-								Output.Add(MakeShared<FSkeletonTreePhysicsConstraintItem>(PhysicsAsset->ConstraintSetup[ConstraintIndex], ConstraintIndex, BoneName, bIsConstraintOnParentBone, PhysicsAsset, SkeletonTreePtr.Pin().ToSharedRef()), BoneName, FSkeletonTreePhysicsBodyItem::GetTypeId());
+
+								Output.Add(MakeShared<FSkeletonTreePhysicsConstraintItem>(PhysicsAsset->ConstraintSetup[ConstraintIndex], ConstraintIndex, BoneName, bIsConstraintOnParentBone, bIsCrossConstraint, PhysicsAsset, SkeletonTreePtr.Pin().ToSharedRef()), BoneName, FSkeletonTreePhysicsBodyItem::GetTypeId());
 							}
 						}
 
diff --git a/Engine/Source/Editor/PhysicsAssetEditor/Private/PhysicsAssetEditorSkeletonTreeBuilder.h b/Engine/Source/Editor/PhysicsAssetEditor/Private/PhysicsAssetEditorSkeletonTreeBuilder.h
--- a/Engine/Source/Editor/PhysicsAssetEditor/Private/PhysicsAssetEditorSkeletonTreeBuilder.h
+++ b/Engine/Source/Editor/PhysicsAssetEditor/Private/PhysicsAssetEditorSkeletonTreeBuilder.h
@@ -22,6 +22,8 @@ public:
 	bool bShowSimulatedBodies;
 	bool bShowConstraints;
 	bool bShowConstraintsOnParentBodies;
+	bool bShowCrossConstraints;
+	bool bShowParentChildConstraints;
 	bool bShowPrimitives;
 
 private:
diff --git a/Engine/Source/Editor/PhysicsAssetEditor/Private/SkeletonTreePhysicsConstraintItem.cpp b/Engine/Source/Editor/PhysicsAssetEditor/Private/SkeletonTreePhysicsConstraintItem.cpp
--- a/Engine/Source/Editor/PhysicsAssetEditor/Private/SkeletonTreePhysicsConstraintItem.cpp
+++ b/Engine/Source/Editor/PhysicsAssetEditor/Private/SkeletonTreePhysicsConstraintItem.cpp
@@ -3,20 +3,26 @@
 #include "SkeletonTreePhysicsConstraintItem.h"
 #include "Styling/AppStyle.h"
 #include "PhysicsAssetRenderUtils.h"
+#include "PhysicsEngine/PhysicsConstraintTemplate.h"
 
 #define LOCTEXT_NAMESPACE "FSkeletonTreePhysicsConstraintItem"
 
-FSkeletonTreePhysicsConstraintItem::FSkeletonTreePhysicsConstraintItem(UPhysicsConstraintTemplate* InConstraint, int32 InConstraintIndex, const FName& InBoneName, bool bInIsConstraintOnParentBody, class UPhysicsAsset* const InPhysicsAsset, const TSharedRef<class ISkeletonTree>& InSkeletonTree)
+FSkeletonTreePhysicsConstraintItem::FSkeletonTreePhysicsConstraintItem(UPhysicsConstraintTemplate* const InConstraint, const int32 InConstraintIndex, const FName& InBoneName, const bool bInIsConstraintOnParentBody, const bool bInIsCrossConstraint, class UPhysicsAsset* const InPhysicsAsset, const TSharedRef<class ISkeletonTree>& InSkeletonTree)
 	: FSkeletonTreePhysicsItem(InPhysicsAsset, InSkeletonTree)
 	, Constraint(InConstraint)
 	, ConstraintIndex(InConstraintIndex)
 	, bIsConstraintOnParentBody(bInIsConstraintOnParentBody)
+	, bIsCrossConstraint(bInIsCrossConstraint)
 {
 	const FConstraintInstance& ConstraintInstance = Constraint->DefaultInstance;
 	FText Label = FText::Format(LOCTEXT("ConstraintNameFormat", "[ {0} -> {1} ] Constraint"), FText::FromName(ConstraintInstance.ConstraintBone2), FText::FromName(ConstraintInstance.ConstraintBone1));
 	DisplayName = *Label.ToString();
 }
 
+FSkeletonTreePhysicsConstraintItem::FSkeletonTreePhysicsConstraintItem(UPhysicsConstraintTemplate* InConstraint, int32 InConstraintIndex, const FName& InBoneName, bool bInIsConstraintOnParentBody, class UPhysicsAsset* const InPhysicsAsset, const TSharedRef<class ISkeletonTree>& InSkeletonTree)
+	: FSkeletonTreePhysicsConstraintItem(InConstraint, InConstraintIndex, InBoneName, bInIsConstraintOnParentBody, false, InPhysicsAsset, InSkeletonTree)
+{}
+
 UObject* FSkeletonTreePhysicsConstraintItem::GetObject() const
 {
 	return Constraint;
@@ -42,7 +48,7 @@ ECheckBoxState FSkeletonTreePhysicsConstraintItem::IsItemDisplayed() const
 
 const FSlateBrush* FSkeletonTreePhysicsConstraintItem::GetBrush() const
 {
-	return	FAppStyle::GetBrush("PhysicsAssetEditor.Tree.Constraint");
+	return	(bIsCrossConstraint) ? FAppStyle::GetBrush("PhysicsAssetEditor.Tree.CrossConstraint") : FAppStyle::GetBrush("PhysicsAssetEditor.Tree.Constraint");
 }
 
 FSlateColor FSkeletonTreePhysicsConstraintItem::GetTextColor() const
diff --git a/Engine/Source/Editor/PhysicsAssetEditor/Private/SkeletonTreePhysicsConstraintItem.h b/Engine/Source/Editor/PhysicsAssetEditor/Private/SkeletonTreePhysicsConstraintItem.h
--- a/Engine/Source/Editor/PhysicsAssetEditor/Private/SkeletonTreePhysicsConstraintItem.h
+++ b/Engine/Source/Editor/PhysicsAssetEditor/Private/SkeletonTreePhysicsConstraintItem.h
@@ -5,15 +5,17 @@
 #include "CoreMinimal.h"
 #include "Styling/SlateColor.h"
 #include "SkeletonTreePhysicsItem.h"
-#include "PhysicsEngine/PhysicsConstraintTemplate.h"
+
+class UPhysicsConstraintTemplate;
 
 class FSkeletonTreePhysicsConstraintItem : public FSkeletonTreePhysicsItem
 {
 public:
 	SKELETON_TREE_ITEM_TYPE(FSkeletonTreePhysicsConstraintItem, FSkeletonTreePhysicsItem)
 
-	FSkeletonTreePhysicsConstraintItem(class UPhysicsConstraintTemplate* InConstraint, int32 InConstraintIndex, const FName& InBoneName, bool bInIsConstraintOnParentBody, class UPhysicsAsset* const InPhysicsAsset, const TSharedRef<class ISkeletonTree>& InSkeletonTree);
-
+	FSkeletonTreePhysicsConstraintItem(UPhysicsConstraintTemplate* const InConstraint, const int32 InConstraintIndex, const FName& InBoneName, const bool bInIsConstraintOnParentBody, const bool bInIsCrossConstraint, class UPhysicsAsset* const InPhysicsAsset, const TSharedRef<class ISkeletonTree>& InSkeletonTree);
+	FSkeletonTreePhysicsConstraintItem(UPhysicsConstraintTemplate* InConstraint, int32 InConstraintIndex, const FName& InBoneName, bool bInIsConstraintOnParentBody, class UPhysicsAsset* const InPhysicsAsset, const TSharedRef<class ISkeletonTree>& InSkeletonTree);
+	
 	/** ISkeletonTreeItem interface */
 	virtual UObject* GetObject() const override;
 
@@ -23,6 +25,9 @@ public:
 	/** since constraint are showing on both parent and child, gets  if this tree item is the one on the parent body */
 	bool IsConstraintOnParentBody() const { return bIsConstraintOnParentBody; }
 
+	/** True if the bones affected by this constraint are not adjacent (parent/child) in the owning skeleton. Constraints between non-adjacent bones are referred to as 'cross-constraints' and have different icons in the editor. */
+	bool IsCrossConstraint() const { return bIsCrossConstraint; }
+
 	/** UI Callbacks */
 	virtual void OnToggleItemDisplayed(ECheckBoxState InCheckboxState) override;
 	virtual ECheckBoxState IsItemDisplayed() const override;
@@ -45,4 +50,7 @@ private:
 
 	/** since constraint are showing on both parent and child, indicates if this tree item is the one on the parent body */
 	bool bIsConstraintOnParentBody;
+
+	/** True if the bones affected by this constraint are not adjacent (parent/child) in the owning skeleton. Constraints between non-adjacent bones are referred to as 'cross-constraints' and have different icons in the editor. */
+	bool bIsCrossConstraint;
 };
