[UE] PhAT Update CoM offsets and icons in skeleton tree correctly when changes are made in the details panel
- Skeleton tree is updated (inc physics body icons) when the physics body type (Static/Simulated/Kinematic) is changed in the details panel
- Changes to physics body position via the details panel now respect the CoM lock axis in component space flags

#rnx

diff --git a/Engine/Source/Editor/PhysicsAssetEditor/Private/PhysicsAssetEditorBodySetupDetailsCustomization.cpp b/Engine/Source/Editor/PhysicsAssetEditor/Private/PhysicsAssetEditorBodySetupDetailsCustomization.cpp
--- a/Engine/Source/Editor/PhysicsAssetEditor/Private/PhysicsAssetEditorBodySetupDetailsCustomization.cpp
+++ b/Engine/Source/Editor/PhysicsAssetEditor/Private/PhysicsAssetEditorBodySetupDetailsCustomization.cpp
@@ -19,6 +19,35 @@
 
 #define LOCTEXT_NAMESPACE "PhysicsAssetBodyDetailsCustomization"
 
+// File Scope Utility Functions //
+
+void OnPhysicsTypePropertyChanged(TSharedPtr<IPropertyHandle> PropertyHandle, FPhysicsAssetEditorSharedData* const SharedData)
+{
+	if (SharedData)
+	{
+		SharedData->BroadcastHierarchyChanged(); // Refresh Skeleton Tree
+	}
+}
+
+void OnGeometryPropertyPreChange(TSharedPtr<IPropertyHandle> PropertyHandle, FPhysicsAssetEditorSharedData* const SharedData)
+{
+	if (SharedData)
+	{
+		SharedData->RecordSelectedBodiesCoM();
+	}
+}
+
+void OnGeometryPropertyChanged(TSharedPtr<IPropertyHandle> PropertyHandle, FPhysicsAssetEditorSharedData* const SharedData)
+{
+	if (SharedData)
+	{
+		SharedData->UpdateSelectedBodiesCoM();
+		SharedData->RefreshPhysicsAssetChange(SharedData->PhysicsAsset, false);
+	}
+}
+
+// class FPhysicsAssetEditorBodySetupDetailsCustomization //
+
 TSharedRef<IDetailCustomization> FPhysicsAssetEditorBodySetupDetailsCustomization::MakeInstance()
 {
 	return MakeShared<FPhysicsAssetEditorBodySetupDetailsCustomization>();
@@ -68,6 +97,21 @@ void FPhysicsAssetEditorBodySetupDetailsCustomization::CustomizeDetails(IDetailL
 	}
 	
 	FBodySetupDetails::CustomizeDetails(DetailLayout);
+
+	{
+		FPhysicsAssetEditorSharedData* const PhysicsAssetSharedData = SharedData;
+
+		if (TSharedPtr<IPropertyHandle> PropertyHandle = DetailLayout.GetProperty(GET_MEMBER_NAME_CHECKED(UBodySetup, PhysicsType), UBodySetupCore::StaticClass()))
+		{
+			PropertyHandle->SetOnPropertyValueChanged(FSimpleDelegate::CreateLambda([PropertyHandle, PhysicsAssetSharedData]() { OnPhysicsTypePropertyChanged(PropertyHandle, PhysicsAssetSharedData); }));
+		}
+
+		if (TSharedPtr<IPropertyHandle> PropertyHandle = DetailLayout.GetProperty(GET_MEMBER_NAME_CHECKED(UBodySetup, AggGeom)))
+		{
+			PropertyHandle->SetOnChildPropertyValuePreChange(FSimpleDelegate::CreateLambda([PropertyHandle, PhysicsAssetSharedData]() { OnGeometryPropertyPreChange(PropertyHandle, PhysicsAssetSharedData); }));
+			PropertyHandle->SetOnChildPropertyValueChanged(FSimpleDelegate::CreateLambda([PropertyHandle, PhysicsAssetSharedData]() { OnGeometryPropertyChanged(PropertyHandle, PhysicsAssetSharedData); }));
+		}
+	}
 }
 
 void FPhysicsAssetEditorBodySetupDetailsCustomization::CustomizeCoMNudge(IDetailLayoutBuilder& DetailBuilder, TSharedRef<IPropertyHandle> BodyInstanceHandler)
diff --git a/Engine/Source/Editor/PhysicsAssetEditor/Private/PhysicsAssetEditorEditMode.cpp b/Engine/Source/Editor/PhysicsAssetEditor/Private/PhysicsAssetEditorEditMode.cpp
--- a/Engine/Source/Editor/PhysicsAssetEditor/Private/PhysicsAssetEditorEditMode.cpp
+++ b/Engine/Source/Editor/PhysicsAssetEditor/Private/PhysicsAssetEditorEditMode.cpp
@@ -173,9 +173,9 @@ bool FPhysicsAssetEditorEditMode::HandleBeginTransform()
 			{
 				SharedData->PhysicsAsset->SkeletalBodySetups[SharedData->SelectedBodies[i].Index]->Modify();
 				SharedData->SelectedBodies[i].ManipulateTM = FTransform::Identity;
-				SharedData->SelectedBodies[i].CoMPosition = SharedData->EditorSkelComp->Bodies[SharedData->SelectedBodies[i].Index]->GetCOMPosition();
 			}
 
+			SharedData->RecordSelectedBodiesCoM();
 			SharedData->bManipulating = true;
 		}
 	
@@ -561,7 +561,7 @@ bool FPhysicsAssetEditorEditMode::InputDelta(FEditorViewportClient* InViewportCl
 					SelectedObject.CoMPosition += DragVec; // Use widget movement to update selected objects target worldspace location, then calculate the correct CoM nudge to position the world space CoM at this location when manipulation ends.
 
 					// Update CoM nudge value in body instance to ensure the values in the details panel update as the manipulator is moved in the viewport.
-					const FVector CalculatedCoMOffset = CalculateCoMNudgeForWorldSpacePosition(BodyIndex, SelectedObject.CoMPosition);
+					const FVector CalculatedCoMOffset = SharedData->CalculateCoMNudgeForWorldSpacePosition(BodyIndex, SelectedObject.CoMPosition);
 					SharedData->EditorSkelComp->Bodies[BodyIndex]->COMNudge = CalculatedCoMOffset;
 					SharedData->PhysicsAsset->SkeletalBodySetups[BodyIndex]->DefaultInstance.COMNudge = CalculatedCoMOffset;
 				}
@@ -1230,38 +1230,12 @@ void FPhysicsAssetEditorEditMode::UpdateCoM()
 {
 	if (SharedData->bShouldUpdatedSelectedCoMs) // < This calculation must be delayed by a frame s.t. changes to the physics state have been propagated to the physics bodies.
 	{
-		TSet<int32> EncounteredBodyIndexs;
-
-		for (int32 i = 0; i < SharedData->SelectedBodies.Num(); ++i)
-		{
-			FPhysicsAssetEditorSharedData::FSelection& SelectedObject = SharedData->SelectedBodies[i];
-			const int32 BodyIndex = SelectedObject.Index;
-			FBodyInstance* const EditorBodyInstance = SharedData->EditorSkelComp->Bodies[BodyIndex];
-
-			if (SharedData->IsCoMAxisFixedInComponentSpace(BodyIndex, EAxis::X) || SharedData->IsCoMAxisFixedInComponentSpace(BodyIndex, EAxis::Y) || SharedData->IsCoMAxisFixedInComponentSpace(BodyIndex, EAxis::Z))
-			{
-				if (!EncounteredBodyIndexs.Contains(BodyIndex)) // Ensure we only update the CoM of each physics body once (SelectedBodies actually contains a list of selected primitives).
-				{
-					EncounteredBodyIndexs.Add(BodyIndex);
-
-					const FVector CoMOffset = EditorBodyInstance->COMNudge;
-
-					FVector	CalculatedCoMOffset = CalculateCoMNudgeForWorldSpacePosition(BodyIndex, SelectedObject.CoMPosition);
-
-					// Only apply lock to the specified Axis in bone space.
-					if (!SharedData->IsCoMAxisFixedInComponentSpace(BodyIndex, EAxis::X)) { CalculatedCoMOffset.X = CoMOffset.X; }
-					if (!SharedData->IsCoMAxisFixedInComponentSpace(BodyIndex, EAxis::Y)) { CalculatedCoMOffset.Y = CoMOffset.Y; }
-					if (!SharedData->IsCoMAxisFixedInComponentSpace(BodyIndex, EAxis::Z)) { CalculatedCoMOffset.Z = CoMOffset.Z; }
-
-					SharedData->PhysicsAsset->SkeletalBodySetups[BodyIndex]->DefaultInstance.COMNudge = CalculatedCoMOffset;
-				}
-			}
-		}
-
+		SharedData->UpdateSelectedBodiesCoM();
+		
 		for (int32 i = 0; i < SharedData->SelectedCoMs.Num(); ++i)
 		{
 			FPhysicsAssetEditorSharedData::FSelection& SelectedObject = SharedData->SelectedCoMs[i];
-			const FVector CalculatedCoMOffset = CalculateCoMNudgeForWorldSpacePosition(SelectedObject.Index, SelectedObject.CoMPosition);
+			const FVector CalculatedCoMOffset = SharedData->CalculateCoMNudgeForWorldSpacePosition(SelectedObject.Index, SelectedObject.CoMPosition);
 			SharedData->PhysicsAsset->SkeletalBodySetups[SelectedObject.Index]->DefaultInstance.COMNudge = CalculatedCoMOffset;
 		}
 
@@ -1271,22 +1245,4 @@ void FPhysicsAssetEditorEditMode::UpdateCoM()
 	}
 }
 
-FVector FPhysicsAssetEditorEditMode::CalculateCoMNudgeForWorldSpacePosition(const int32 BodyIndex, const FVector& CoMPositionWorldSpace) const
-{
-	FVector	CalculatedCoMOffset = FVector::ZeroVector;
-
-	if (SharedData->EditorSkelComp->Bodies.IsValidIndex(BodyIndex))
-	{
-		if (FBodyInstance* const EditorBodyInstance = SharedData->EditorSkelComp->Bodies[BodyIndex])
-		{
-			const int32 BoneIndex = SharedData->EditorSkelComp->GetBoneIndex(SharedData->PhysicsAsset->SkeletalBodySetups[BodyIndex]->BoneName);
-			const FTransform BoneTM = SharedData->EditorSkelComp->GetBoneTransform(BoneIndex);
-			const FVector CoMWithoutNudge = EditorBodyInstance->GetMassSpaceLocal().GetTranslation() - EditorBodyInstance->COMNudge;
-			CalculatedCoMOffset = BoneTM.InverseTransformPosition(CoMPositionWorldSpace) - CoMWithoutNudge;
-		}
-	}
-
-	return CalculatedCoMOffset;
-}
-
 #undef LOCTEXT_NAMESPACE
diff --git a/Engine/Source/Editor/PhysicsAssetEditor/Private/PhysicsAssetEditorEditMode.h b/Engine/Source/Editor/PhysicsAssetEditor/Private/PhysicsAssetEditorEditMode.h
--- a/Engine/Source/Editor/PhysicsAssetEditor/Private/PhysicsAssetEditorEditMode.h
+++ b/Engine/Source/Editor/PhysicsAssetEditor/Private/PhysicsAssetEditorEditMode.h
@@ -93,9 +93,6 @@ private:
 	// Update the Center of Mass position after body transforms have been manipulated in the view port.
 	void UpdateCoM();
 
-	// Calculate a Center of Mass nudge (offset) for a given body that will locate that body's CoM at the supplied position in world space.
-	FVector CalculateCoMNudgeForWorldSpacePosition(const int32 BodyIndex, const FVector& CoMPositionWorldSpace) const;
-
 private:
 	/** Shared data */
 	FPhysicsAssetEditorSharedData* SharedData;
diff --git a/Engine/Source/Editor/PhysicsAssetEditor/Private/PhysicsAssetEditorSharedData.cpp b/Engine/Source/Editor/PhysicsAssetEditor/Private/PhysicsAssetEditorSharedData.cpp
--- a/Engine/Source/Editor/PhysicsAssetEditor/Private/PhysicsAssetEditorSharedData.cpp
+++ b/Engine/Source/Editor/PhysicsAssetEditor/Private/PhysicsAssetEditorSharedData.cpp
@@ -415,6 +415,66 @@ void FPhysicsAssetEditorSharedData::SetCoMAxisFixedInComponentSpace(const int32
 	}
 }
 
+FVector FPhysicsAssetEditorSharedData::CalculateCoMNudgeForWorldSpacePosition(const int32 BodyIndex, const FVector& CoMPositionWorldSpace) const
+{
+	FVector	CalculatedCoMOffset = FVector::ZeroVector;
+
+	if (EditorSkelComp && EditorSkelComp->Bodies.IsValidIndex(BodyIndex))
+	{
+		if (FBodyInstance* const EditorBodyInstance = EditorSkelComp->Bodies[BodyIndex])
+		{
+			const int32 BoneIndex = EditorSkelComp->GetBoneIndex(PhysicsAsset->SkeletalBodySetups[BodyIndex]->BoneName);
+			const FTransform BoneTM = EditorSkelComp->GetBoneTransform(BoneIndex);
+			const FVector CoMWithoutNudge = EditorBodyInstance->GetMassSpaceLocal().GetTranslation() - EditorBodyInstance->COMNudge;
+			CalculatedCoMOffset = BoneTM.InverseTransformPosition(CoMPositionWorldSpace) - CoMWithoutNudge;
+		}
+	}
+
+	return CalculatedCoMOffset;
+}
+
+void FPhysicsAssetEditorSharedData::RecordSelectedBodiesCoM()
+{
+	if (EditorSkelComp)
+	{
+		for(FSelection& Body : SelectedBodies)
+		{
+			Body.CoMPosition = EditorSkelComp->Bodies[Body.Index]->GetCOMPosition();
+		}
+	}
+}
+
+void FPhysicsAssetEditorSharedData::UpdateSelectedBodiesCoM()
+{
+	TSet<int32> EncounteredBodyIndexs;
+
+	// Update CoM nudge to compensate for any change in body transform on any axis that is fixed in component space.
+	for (int32 i = 0; i < SelectedBodies.Num(); ++i)
+	{
+		FPhysicsAssetEditorSharedData::FSelection& SelectedObject = SelectedBodies[i];
+		const int32 BodyIndex = SelectedObject.Index;
+		FBodyInstance* const EditorBodyInstance = EditorSkelComp->Bodies[BodyIndex];
+
+		if (IsCoMAxisFixedInComponentSpace(BodyIndex, EAxis::X) || IsCoMAxisFixedInComponentSpace(BodyIndex, EAxis::Y) || IsCoMAxisFixedInComponentSpace(BodyIndex, EAxis::Z))
+		{
+			// Ensure we only update the CoM of each physics body once (SelectedBodies actually contains a list of selected primitives).
+			if (!EncounteredBodyIndexs.Contains(BodyIndex)) 
+			{
+				EncounteredBodyIndexs.Add(BodyIndex);
+				const FVector CoMOffset = EditorBodyInstance->COMNudge;
+				FVector	CalculatedCoMOffset = CalculateCoMNudgeForWorldSpacePosition(BodyIndex, SelectedObject.CoMPosition);
+
+				// Only apply lock to the specified Axis in bone space.
+				if (!IsCoMAxisFixedInComponentSpace(BodyIndex, EAxis::X)) { CalculatedCoMOffset.X = CoMOffset.X; }
+				if (!IsCoMAxisFixedInComponentSpace(BodyIndex, EAxis::Y)) { CalculatedCoMOffset.Y = CoMOffset.Y; }
+				if (!IsCoMAxisFixedInComponentSpace(BodyIndex, EAxis::Z)) { CalculatedCoMOffset.Z = CoMOffset.Z; }
+
+				PhysicsAsset->SkeletalBodySetups[BodyIndex]->DefaultInstance.COMNudge = CalculatedCoMOffset;
+			}
+		}
+	}
+}
+
 bool FPhysicsAssetEditorSharedData::ParseClipboard(UPhysicsAsset*& OutAsset, FString& OutObjectType, UObject*& OutObject)
 {
 	FString ClipboardContent;
diff --git a/Engine/Source/Editor/PhysicsAssetEditor/Private/PhysicsAssetEditorSharedData.h b/Engine/Source/Editor/PhysicsAssetEditor/Private/PhysicsAssetEditorSharedData.h
--- a/Engine/Source/Editor/PhysicsAssetEditor/Private/PhysicsAssetEditorSharedData.h
+++ b/Engine/Source/Editor/PhysicsAssetEditor/Private/PhysicsAssetEditorSharedData.h
@@ -299,6 +299,15 @@ public:
 	bool IsCoMAxisFixedInComponentSpace(const int32 BodyIndex, const EAxis::Type InAxis) const;
 	void SetCoMAxisFixedInComponentSpace(const int32 BodyIndex, const EAxis::Type InAxis, const bool bValue);
 
+	// Calculate a Center of Mass nudge (offset) for a given body that will locate that body's CoM at the supplied position in world space.
+	FVector CalculateCoMNudgeForWorldSpacePosition(const int32 BodyIndex, const FVector& CoMPositionWorldSpace) const;
+
+	// Make a copy of the current component space CoM position from each selected physics body - called before a change in physics body transform
+	void RecordSelectedBodiesCoM();
+	
+	// Update the CoM nudge for each selected physics body - called after a change in physics body transform.
+	void UpdateSelectedBodiesCoM();
+
 private:
 	/** Initializes a constraint setup */
 	void InitConstraintSetup(UPhysicsConstraintTemplate* ConstraintSetup, int32 ChildBodyIndex, int32 ParentBodyIndex);
