[UE] PhAT CoM tools Filter kinematic Bodies and Multi-Selection Support
- Optionally hide CoM Values and markers in Viewport for kinematic bodies
- Ability to tick CoM local vs global space on multiple selected bodies

diff --git a/Engine/Source/Editor/Persona/Private/PhysicsAssetRenderUtils.cpp b/Engine/Source/Editor/Persona/Private/PhysicsAssetRenderUtils.cpp
--- a/Engine/Source/Editor/Persona/Private/PhysicsAssetRenderUtils.cpp
+++ b/Engine/Source/Editor/Persona/Private/PhysicsAssetRenderUtils.cpp
@@ -73,6 +73,7 @@ FPhysicsAssetRenderSettings::FPhysicsAssetRenderSettings()
 	, bShowCOM_DEPRECATED(false)
 	, bShowConstraintsAsPoints(false)
 	, bDrawViolatedLimits(false)
+	, bHideCenterOfMassForKinematicBodies(true)
 	, BoneUnselectedColor(170, 155, 225)
 	, NoCollisionColor(200, 200, 200)
 	, COMRenderColor(255, 255, 100)
@@ -284,6 +285,7 @@ void FPhysicsAssetRenderSettings::ResetEditorViewportOptions()
 	bRenderOnlySelectedConstraints = DefaultObject.bRenderOnlySelectedConstraints;
 	bShowConstraintsAsPoints = DefaultObject.bShowConstraintsAsPoints;
 	bDrawViolatedLimits = DefaultObject.bDrawViolatedLimits;
+	bHideCenterOfMassForKinematicBodies = DefaultObject.bHideCenterOfMassForKinematicBodies;
 }
 
 namespace PhysicsAssetRender
@@ -454,7 +456,7 @@ namespace PhysicsAssetRender
 		}
 	}
 
-	void DebugDrawCenterOfMass(USkeletalMeshComponent* const SkeletalMeshComponent, class UPhysicsAsset* const PhysicsAsset, FPrimitiveDrawInterface* PDI, TFunctionRef< FVector(const uint32) > GetCoMPosition, TFunctionRef< bool(const uint32) > IsSelected, CreateCoMHitProxyFn CreateHitProxy)
+	void DebugDrawCenterOfMass(USkeletalMeshComponent* const SkeletalMeshComponent, UPhysicsAsset* const PhysicsAsset, FPrimitiveDrawInterface* PDI, TFunctionRef< FVector(const uint32) > GetCoMPosition, TFunctionRef< bool(const uint32) > IsSelected, CreateCoMHitProxyFn CreateHitProxy)
 	{
 		static const float SelectedItemRenderSizeMultiplier = 1.5f;
 
@@ -472,7 +474,7 @@ namespace PhysicsAssetRender
 				{
 					const FBodyInstance* const BodyInstance = SkeletalMeshComponent->Bodies[BodyIndex];
 
-					if (BodyInstance && BodyInstance->IsValidBodyInstance())
+					if (BodyInstance && BodyInstance->IsValidBodyInstance() && !(RenderSettings->bHideCenterOfMassForKinematicBodies && IsBodyKinematic(PhysicsAsset, BodyIndex)))
 					{
 						float COMRenderSize = RenderSettings->COMRenderSize;
 						float COMRenderLineThickness = RenderSettings->COMRenderLineThickness;
@@ -838,4 +840,14 @@ void FPhysicsAssetRenderInterface::SetDisplayConstraintTransformComponentRelativ
 	}
 }
 
+bool IsBodyKinematic(const UPhysicsAsset* const PhysicsAsset, const int32 BodyIndex)
+{
+	if (PhysicsAsset && PhysicsAsset->SkeletalBodySetups.IsValidIndex(BodyIndex))
+	{
+		return PhysicsAsset->SkeletalBodySetups[BodyIndex]->PhysicsType == EPhysicsType::PhysType_Kinematic;
+	}
+
+	return false;
+}
+
 #undef LOCTEXT_NAMESPACE
diff --git a/Engine/Source/Editor/Persona/Public/PhysicsAssetRenderUtils.h b/Engine/Source/Editor/Persona/Public/PhysicsAssetRenderUtils.h
--- a/Engine/Source/Editor/Persona/Public/PhysicsAssetRenderUtils.h
+++ b/Engine/Source/Editor/Persona/Public/PhysicsAssetRenderUtils.h
@@ -118,6 +118,9 @@ public:
 
 	UPROPERTY()
 	bool bDrawViolatedLimits;
+	
+	UPROPERTY()
+	bool bHideCenterOfMassForKinematicBodies;
 
 	// Draw colors
 	UPROPERTY()
@@ -256,4 +259,6 @@ public:
 
 	/** Change how the constraint transform component specified by the flags should be displayed, in the frame of their associated physics body or as an offset from the default (snapped) transforms. */
 	virtual void SetDisplayConstraintTransformComponentRelativeToDefault(class UPhysicsAsset* const PhysicsAsset, const EConstraintTransformComponentFlags ComponentFlags, const bool bShouldDisplayRelativeToDefault) override;
-};
\ No newline at end of file
+};
+
+PERSONA_API bool IsBodyKinematic(const UPhysicsAsset* const PhysicsAsset, const int32 BodyIndex); // TODO - where is the best place for this fn to live ?
\ No newline at end of file
diff --git a/Engine/Source/Editor/PhysicsAssetEditor/Private/PhysicsAssetEditor.cpp b/Engine/Source/Editor/PhysicsAssetEditor/Private/PhysicsAssetEditor.cpp
--- a/Engine/Source/Editor/PhysicsAssetEditor/Private/PhysicsAssetEditor.cpp
+++ b/Engine/Source/Editor/PhysicsAssetEditor/Private/PhysicsAssetEditor.cpp
@@ -898,6 +898,7 @@ void FPhysicsAssetEditor::ExtendViewportMenus()
 						{
 							FToolMenuSection& Section = InSubMenu->AddSection("PhysicsAssetEditorCenterOfMassRenderSettings", LOCTEXT("CenterOfMassRenderSettingsHeader", "Center of Mass Drawing"));
 							Section.AddMenuEntry(Commands.DrawBodyMass);
+							Section.AddMenuEntry(Commands.HideCenterOfMassForKinematicBodies);
 							Section.AddEntry(FToolMenuEntry::InitWidget(TEXT("CoMMarkerScale"), WeakPhysicsAssetEditor.Pin()->MakeCoMMarkerScaleWidget(), LOCTEXT("CoMMarkerScaleLabel", "Marker Scale")));
 						}
 
@@ -1689,6 +1690,12 @@ void FPhysicsAssetEditor::BindCommands()
 		FCanExecuteAction(),
 		FIsActionChecked::CreateSP(this, &FPhysicsAssetEditor::IsDrawingBodyMass));
 
+	ViewportCommandList->MapAction(
+		Commands.HideCenterOfMassForKinematicBodies,
+		FExecuteAction::CreateSP(this, &FPhysicsAssetEditor::ToggleHideCenterOfMassForKinematicBodies),
+		FCanExecuteAction(),
+		FIsActionChecked::CreateSP(this, &FPhysicsAssetEditor::IsHidingCenterOfMassForKinematicBodies));
+
 	ViewportCommandList->MapAction(
 		Commands.DrawConstraintsAsPoints,
 		FExecuteAction::CreateSP(this, &FPhysicsAssetEditor::ToggleDrawConstraintsAsPoints),
@@ -3065,6 +3072,17 @@ bool FPhysicsAssetEditor::IsDrawingViolatedLimits() const
 	return SharedData->EditorOptions->bDrawViolatedLimits;
 }
 
+void FPhysicsAssetEditor::ToggleHideCenterOfMassForKinematicBodies()
+{
+	SharedData->EditorOptions->bHideCenterOfMassForKinematicBodies = !SharedData->EditorOptions->bHideCenterOfMassForKinematicBodies;
+	SharedData->EditorOptions->SaveConfig();
+}
+
+bool FPhysicsAssetEditor::IsHidingCenterOfMassForKinematicBodies() const
+{
+	return SharedData->EditorOptions->bHideCenterOfMassForKinematicBodies;
+}
+
 void FPhysicsAssetEditor::ToggleRenderOnlySelectedConstraints()
 {
 	SharedData->EditorOptions->bRenderOnlySelectedConstraints = !SharedData->EditorOptions->bRenderOnlySelectedConstraints;
diff --git a/Engine/Source/Editor/PhysicsAssetEditor/Private/PhysicsAssetEditor.h b/Engine/Source/Editor/PhysicsAssetEditor/Private/PhysicsAssetEditor.h
--- a/Engine/Source/Editor/PhysicsAssetEditor/Private/PhysicsAssetEditor.h
+++ b/Engine/Source/Editor/PhysicsAssetEditor/Private/PhysicsAssetEditor.h
@@ -257,6 +257,8 @@ private:
 	bool IsDrawingConstraintsAsPoints() const;
 	void ToggleDrawViolatedLimits();
 	bool IsDrawingViolatedLimits() const;
+	void ToggleHideCenterOfMassForKinematicBodies();
+	bool IsHidingCenterOfMassForKinematicBodies() const;
 	void ToggleRenderOnlySelectedConstraints();
 	bool IsRenderingOnlySelectedConstraints() const;
 	void ToggleRenderOnlySelectedSolid();
diff --git a/Engine/Source/Editor/PhysicsAssetEditor/Private/PhysicsAssetEditorActions.cpp b/Engine/Source/Editor/PhysicsAssetEditor/Private/PhysicsAssetEditorActions.cpp
--- a/Engine/Source/Editor/PhysicsAssetEditor/Private/PhysicsAssetEditorActions.cpp
+++ b/Engine/Source/Editor/PhysicsAssetEditor/Private/PhysicsAssetEditorActions.cpp
@@ -68,6 +68,7 @@ void FPhysicsAssetEditorCommands::RegisterCommands()
 	UI_COMMAND(DrawConstraintsAsPoints, "Draw Constraints As Points", "Draw Constraints As Points", EUserInterfaceActionType::ToggleButton, FInputChord());
 	UI_COMMAND(DrawViolatedLimits, "Draw Violated Limits", "Draw Violated Limits", EUserInterfaceActionType::ToggleButton, FInputChord());
 	UI_COMMAND(DrawBodyMass, "Mass Values", "Draw text displaying body mass next to CoM", EUserInterfaceActionType::ToggleButton, FInputChord());
+	UI_COMMAND(HideCenterOfMassForKinematicBodies, "Hide Kinematic Bodies", "Hide CoM position and mass value for Kinematic Bodies", EUserInterfaceActionType::ToggleButton, FInputChord());
 	UI_COMMAND(DisableCollision, "Disable Collision", "Disable collision between the currently selected bodies", EUserInterfaceActionType::Button, FInputChord(EKeys::RightBracket));
 	UI_COMMAND(DisableCollisionAll, "Disable Collision All", "Disable collision between the currently selected bodies and all bodies", EUserInterfaceActionType::Button, FInputChord(EModifierKey::Control, EKeys::RightBracket));
 	UI_COMMAND(EnableCollision, "Enable Collision", "Enable collision between the currently selected bodies", EUserInterfaceActionType::Button, FInputChord(EKeys::LeftBracket));
diff --git a/Engine/Source/Editor/PhysicsAssetEditor/Private/PhysicsAssetEditorActions.h b/Engine/Source/Editor/PhysicsAssetEditor/Private/PhysicsAssetEditorActions.h
--- a/Engine/Source/Editor/PhysicsAssetEditor/Private/PhysicsAssetEditorActions.h
+++ b/Engine/Source/Editor/PhysicsAssetEditor/Private/PhysicsAssetEditorActions.h
@@ -79,6 +79,7 @@ public:
 	TSharedPtr<FUICommandInfo> HideKinematicBodies;
 	TSharedPtr<FUICommandInfo> DrawBodyMass;
 	TSharedPtr<FUICommandInfo> DrawCenterOfMass;
+	TSharedPtr<FUICommandInfo> HideCenterOfMassForKinematicBodies;
 	TSharedPtr<FUICommandInfo> RenderOnlySelectedBodyCoMs;
 	TSharedPtr<FUICommandInfo> RenderOnlySelectedConstraints;
 	TSharedPtr<FUICommandInfo> DrawConstraintsAsPoints;
diff --git a/Engine/Source/Editor/PhysicsAssetEditor/Private/PhysicsAssetEditorBodySetupDetailsCustomization.cpp b/Engine/Source/Editor/PhysicsAssetEditor/Private/PhysicsAssetEditorBodySetupDetailsCustomization.cpp
--- a/Engine/Source/Editor/PhysicsAssetEditor/Private/PhysicsAssetEditorBodySetupDetailsCustomization.cpp
+++ b/Engine/Source/Editor/PhysicsAssetEditor/Private/PhysicsAssetEditorBodySetupDetailsCustomization.cpp
@@ -55,14 +55,12 @@ TSharedRef<IDetailCustomization> FPhysicsAssetEditorBodySetupDetailsCustomizatio
 
 FPhysicsAssetEditorBodySetupDetailsCustomization::FPhysicsAssetEditorBodySetupDetailsCustomization()
 : SharedData(nullptr)
-, BodyIndex(INDEX_NONE)
 {}
 
 void FPhysicsAssetEditorBodySetupDetailsCustomization::CustomizeDetails(IDetailLayoutBuilder& DetailLayout)
 {
 	const UPhysicsAsset* PhysicsAsset = nullptr;
-	BodyIndex = INDEX_NONE;
-
+	
 	// Find the physics asset that owns the BodySetup that is to be customized.
 	{
 		TArray< TWeakObjectPtr<UObject> > ObjectsToBeCustomized;
@@ -75,12 +73,7 @@ void FPhysicsAssetEditorBodySetupDetailsCustomization::CustomizeDetails(IDetailL
 				if (USkeletalBodySetup* const SkeletalBodySetup = Cast<USkeletalBodySetup>(ObjectPtr.Get()))
 				{
 					PhysicsAsset = Cast<UPhysicsAsset>(SkeletalBodySetup->GetOuter());
-
-					if (PhysicsAsset)
-					{
-						BodyIndex = PhysicsAsset->FindBodyIndex(SkeletalBodySetup->BoneName);
-						break;
-					}
+					break;
 				}
 			}
 		}
@@ -95,7 +88,7 @@ void FPhysicsAssetEditorBodySetupDetailsCustomization::CustomizeDetails(IDetailL
 			break;
 		}
 	}
-	
+
 	FBodySetupDetails::CustomizeDetails(DetailLayout);
 
 	{
@@ -166,9 +159,9 @@ void FPhysicsAssetEditorBodySetupDetailsCustomization::CustomizeCoMNudge(IDetail
 							[
 								SNew(SImage)
 									.ColorAndOpacity(FSlateColor::UseForeground())
-									.Image_Lambda([this, Axis, WorldSpaceIcon, LocalSpaceIcon]() { return (this->IsCOMFixedInComponentSpace(Axis)) ? WorldSpaceIcon.GetIcon() : LocalSpaceIcon.GetIcon(); })
+									.Image_Lambda([this, Axis, WorldSpaceIcon, LocalSpaceIcon]() { return (this->IsCOMAxisFixedInComponentSpace(Axis)) ? WorldSpaceIcon.GetIcon() : LocalSpaceIcon.GetIcon(); })
 									.ToolTipText_Lambda([this, Axis]()
-										{ return (this->IsCOMFixedInComponentSpace(Axis))
+										{ return (this->IsCOMAxisFixedInComponentSpace(Axis))
 										? LOCTEXT("ToolTipDeactivateCOMFixedInComponentSpace", "Center of Mass position is currently fixed in component space on this axis. When the physics body is moved the CoM offset will automatically update to maintain the current position. This will only have an effect in the editor. Click to toggle this behavior (Shift + Click to toggle all axis).")
 										: LOCTEXT("ToolTipActivateCOMFixedInComponentSpace", "Center of Mass position is currently fixed in local space on this axis. When the physics body is moved the CoM offset will move with it to maintain the current position relative to the body. This will only have an effect in the editor. Click to toggle this behavior (Shift + Click to toggle all axis)."); })
 							]
@@ -212,35 +205,46 @@ void FPhysicsAssetEditorBodySetupDetailsCustomization::CustomizeCoMNudge(IDetail
 
 FReply FPhysicsAssetEditorBodySetupDetailsCustomization::ToggleFixCOMInComponentSpace(const EAxis::Type Axis)
 {
-	if (SharedData)
-	{
-		const bool bIsCOMFixed = SharedData->IsCoMAxisFixedInComponentSpace(BodyIndex, Axis);
+	const bool bIsCOMFixed = IsCOMAxisFixedInComponentSpace(Axis);
 
-		if (FSlateApplication::Get().GetModifierKeys().IsShiftDown())
-		{
-			SharedData->SetCoMAxisFixedInComponentSpace(BodyIndex, EAxis::X, !bIsCOMFixed);
-			SharedData->SetCoMAxisFixedInComponentSpace(BodyIndex, EAxis::Y, !bIsCOMFixed);
-			SharedData->SetCoMAxisFixedInComponentSpace(BodyIndex, EAxis::Z, !bIsCOMFixed);
-		}
-		else
-		{
-			SharedData->SetCoMAxisFixedInComponentSpace(BodyIndex, Axis, !bIsCOMFixed);
-		}
+	if (FSlateApplication::Get().GetModifierKeys().IsShiftDown())
+	{
+		SetCoMAxisFixedInComponentSpace(EAxis::X, !bIsCOMFixed);
+		SetCoMAxisFixedInComponentSpace(EAxis::Y, !bIsCOMFixed);
+		SetCoMAxisFixedInComponentSpace(EAxis::Z, !bIsCOMFixed);
 	}
-
+	else
+	{
+		SetCoMAxisFixedInComponentSpace(Axis, !bIsCOMFixed);
+	}
+	
 	return FReply::Handled();
 }
 
-bool FPhysicsAssetEditorBodySetupDetailsCustomization::IsCOMFixedInComponentSpace(const EAxis::Type Axis) const
+bool FPhysicsAssetEditorBodySetupDetailsCustomization::IsCOMAxisFixedInComponentSpace(const EAxis::Type Axis) const
 {
-	bool bIsCOMFixed = false;
+	bool bIsCOMFixed = true;
 
 	if (SharedData)
 	{
-		bIsCOMFixed = SharedData->IsCoMAxisFixedInComponentSpace(BodyIndex, Axis);;
+		for (const FPhysicsAssetEditorSharedData::FSelection& SelectedBody : SharedData->GetUniqueSelectedBodies())
+		{
+			bIsCOMFixed &= SharedData->IsCoMAxisFixedInComponentSpace(SelectedBody.Index, Axis);
+		}
 	}
 
 	return bIsCOMFixed;
 }
 
+void FPhysicsAssetEditorBodySetupDetailsCustomization::SetCoMAxisFixedInComponentSpace(const EAxis::Type Axis, const bool bCOMFixed) const
+{
+	if (SharedData)
+	{
+		for (const FPhysicsAssetEditorSharedData::FSelection& SelectedBody : SharedData->GetUniqueSelectedBodies())
+		{
+			SharedData->SetCoMAxisFixedInComponentSpace(SelectedBody.Index, Axis, bCOMFixed);
+		}
+	}
+}
+
 #undef LOCTEXT_NAMESPACE
diff --git a/Engine/Source/Editor/PhysicsAssetEditor/Private/PhysicsAssetEditorBodySetupDetailsCustomization.h b/Engine/Source/Editor/PhysicsAssetEditor/Private/PhysicsAssetEditorBodySetupDetailsCustomization.h
--- a/Engine/Source/Editor/PhysicsAssetEditor/Private/PhysicsAssetEditorBodySetupDetailsCustomization.h
+++ b/Engine/Source/Editor/PhysicsAssetEditor/Private/PhysicsAssetEditorBodySetupDetailsCustomization.h
@@ -34,10 +34,11 @@ public:
 
 	virtual void CustomizeCoMNudge(IDetailLayoutBuilder& DetailBuilder, TSharedRef<IPropertyHandle> BodyInstanceHandler) override;
 	FReply ToggleFixCOMInComponentSpace(const EAxis::Type Axis);
-	bool IsCOMFixedInComponentSpace(const EAxis::Type Axis) const;
+	bool IsCOMAxisFixedInComponentSpace(const EAxis::Type Axis) const;
+	void SetCoMAxisFixedInComponentSpace(const EAxis::Type Axis, const bool bCoMFixed) const;
 
 private:
 
 	FPhysicsAssetEditorSharedData* SharedData;
-	int32 BodyIndex;
+
 };
diff --git a/Engine/Source/Editor/PhysicsAssetEditor/Private/PhysicsAssetEditorEditMode.cpp b/Engine/Source/Editor/PhysicsAssetEditor/Private/PhysicsAssetEditorEditMode.cpp
--- a/Engine/Source/Editor/PhysicsAssetEditor/Private/PhysicsAssetEditorEditMode.cpp
+++ b/Engine/Source/Editor/PhysicsAssetEditor/Private/PhysicsAssetEditorEditMode.cpp
@@ -707,7 +707,7 @@ void FPhysicsAssetEditorEditMode::DrawHUD(FEditorViewportClient* ViewportClient,
 			// iterate over each body
 			for (int32 BodyIndex = 0, BodyCount = SharedData->EditorSkelComp->Bodies.Num(); BodyIndex < BodyCount; ++BodyIndex)
 			{
-				if (!bDrawSelectedOnly || SharedData->IsBodySelected(BodyIndex))
+				if ((!bDrawSelectedOnly || SharedData->IsBodySelected(BodyIndex)) && !(RenderSettings->bHideCenterOfMassForKinematicBodies && IsBodyKinematic(SharedData->PhysicsAsset.Get(), BodyIndex)))
 				{
 					FBodyInstance* BodyInst = SharedData->EditorSkelComp->Bodies[BodyIndex];
 					check(BodyInst);
diff --git a/Engine/Source/Editor/PhysicsAssetEditor/Private/PhysicsAssetEditorSharedData.cpp b/Engine/Source/Editor/PhysicsAssetEditor/Private/PhysicsAssetEditorSharedData.cpp
--- a/Engine/Source/Editor/PhysicsAssetEditor/Private/PhysicsAssetEditorSharedData.cpp
+++ b/Engine/Source/Editor/PhysicsAssetEditor/Private/PhysicsAssetEditorSharedData.cpp
@@ -3293,6 +3293,25 @@ void FPhysicsAssetEditorSharedData::UpdateClothPhysics()
 	}
 }
 
+TArray<FPhysicsAssetEditorSharedData::FSelection> FPhysicsAssetEditorSharedData::GetUniqueSelectedBodies() const
+{
+	TArray<FSelection> UniqueSelectedBodies;
+
+	TArray<int32> EncounteredBodyIndexs;
+
+	for (const FSelection& Selection : SelectedBodies)
+	{
+		if (!EncounteredBodyIndexs.Contains(Selection.Index))
+		{
+			UniqueSelectedBodies.Add(Selection);
+			EncounteredBodyIndexs.Add(Selection.Index);
+		}
+	}
+
+	return UniqueSelectedBodies;
+}
+
+
 FPhysicsAssetRenderSettings* FPhysicsAssetEditorSharedData::GetRenderSettings() const
 {
 	return UPhysicsAssetRenderUtilities::GetSettings(PhysicsAsset);
diff --git a/Engine/Source/Editor/PhysicsAssetEditor/Private/PhysicsAssetEditorSharedData.h b/Engine/Source/Editor/PhysicsAssetEditor/Private/PhysicsAssetEditorSharedData.h
--- a/Engine/Source/Editor/PhysicsAssetEditor/Private/PhysicsAssetEditorSharedData.h
+++ b/Engine/Source/Editor/PhysicsAssetEditor/Private/PhysicsAssetEditorSharedData.h
@@ -402,6 +402,13 @@ public:
 		return Count ? &SelectedBodies[Count - 1] : NULL;
 	}
 
+	/** 
+	 * Returns an array containing only the first FSelection that references a particular body 
+	 * (The SelectedBodies Array includes an FSelction for every selected primitive so can 
+	 * include multiple FSelections with the same body index)
+	 */
+	TArray<FSelection> GetUniqueSelectedBodies() const;
+
 	FSelection* GetSelectedCoM()
 	{
 		int32 Count = SelectedCoMs.Num();
diff --git a/Engine/Source/Editor/PhysicsAssetEditor/Private/PhysicsAssetEditorSkeletalMeshComponent.cpp b/Engine/Source/Editor/PhysicsAssetEditor/Private/PhysicsAssetEditorSkeletalMeshComponent.cpp
--- a/Engine/Source/Editor/PhysicsAssetEditor/Private/PhysicsAssetEditorSkeletalMeshComponent.cpp
+++ b/Engine/Source/Editor/PhysicsAssetEditor/Private/PhysicsAssetEditorSkeletalMeshComponent.cpp
@@ -131,6 +131,7 @@ void UPhysicsAssetEditorSkeletalMeshComponent::DebugDraw(const FSceneView* View,
 		RenderSettings->bRenderOnlySelectedConstraints = SharedData->EditorOptions->bRenderOnlySelectedConstraints;
 		RenderSettings->bShowConstraintsAsPoints = SharedData->EditorOptions->bShowConstraintsAsPoints;
 		RenderSettings->bDrawViolatedLimits = SharedData->EditorOptions->bDrawViolatedLimits;
+		RenderSettings->bHideCenterOfMassForKinematicBodies = SharedData->EditorOptions->bHideCenterOfMassForKinematicBodies;
 
 		// Draw Bodies.
 		{
diff --git a/Engine/Source/Editor/UnrealEd/Classes/Preferences/PhysicsAssetEditorOptions.h b/Engine/Source/Editor/UnrealEd/Classes/Preferences/PhysicsAssetEditorOptions.h
--- a/Engine/Source/Editor/UnrealEd/Classes/Preferences/PhysicsAssetEditorOptions.h
+++ b/Engine/Source/Editor/UnrealEd/Classes/Preferences/PhysicsAssetEditorOptions.h
@@ -119,6 +119,10 @@ class UPhysicsAssetEditorOptions : public UObject
 	UPROPERTY(config)
 	uint32 bDrawViolatedLimits:1;
 
+	/** Whether to draw CoM positions / mass values for kinematic bodies as well as other body types */
+	UPROPERTY(config)
+	uint32 bHideCenterOfMassForKinematicBodies:1;
+
 	/** Whether to only render selected constraints */
 	UPROPERTY(config)
 	uint32 bRenderOnlySelectedConstraints:1;
diff --git a/Engine/Source/Editor/UnrealEd/Private/PreferenceStubs.cpp b/Engine/Source/Editor/UnrealEd/Private/PreferenceStubs.cpp
--- a/Engine/Source/Editor/UnrealEd/Private/PreferenceStubs.cpp
+++ b/Engine/Source/Editor/UnrealEd/Private/PreferenceStubs.cpp
@@ -44,6 +44,7 @@ UPhysicsAssetEditorOptions::UPhysicsAssetEditorOptions(const FObjectInitializer&
 
 	bShowConstraintsAsPoints = false;
 	bDrawViolatedLimits = false;
+	bHideCenterOfMassForKinematicBodies = false;
 	bSimulationFloorCollisionEnabled = true;
 	COMRenderSize = 1.0f;
 	ConstraintDrawSize = 1.0f;
