[UE] PhAT Move CoM nudge axis fixed in local/component space flags from PhysicsAssetSharedData to PhysicsAsset (Editor Only Data) so that settings can be serialized with the asset and persist between editor sessions.

diff --git a/Engine/Source/Editor/PhysicsAssetEditor/Private/PhysicsAssetEditorSharedData.cpp b/Engine/Source/Editor/PhysicsAssetEditor/Private/PhysicsAssetEditorSharedData.cpp
--- a/Engine/Source/Editor/PhysicsAssetEditor/Private/PhysicsAssetEditorSharedData.cpp
+++ b/Engine/Source/Editor/PhysicsAssetEditor/Private/PhysicsAssetEditorSharedData.cpp
@@ -53,6 +53,8 @@ namespace SharedDataConstants
 	const FString BodyType = TEXT("SkeletalBodySetup");
 }
 
+// File Scope Utility Functions //
+
 bool SetSelected(TArray<FPhysicsAssetEditorSharedData::FSelection>& SelectedElements, const TArray<FPhysicsAssetEditorSharedData::FSelection>& DeltaElements, const bool bSelected)
 {
 	if (DeltaElements.Num() == 0)
@@ -91,6 +93,14 @@ bool SetSelected(TArray<FPhysicsAssetEditorSharedData::FSelection>& SelectedElem
 	return SetSelected(SelectedElements, DeltaElements, bSelected);
 }
 
+/** Returns the Editor Body Flag bit mask that indicates if the supplied axis has been fixed in component space */
+int32 FindCoMAxisEditorBodyFlag(const EAxis::Type InAxis)
+{
+	return int32(1) << (InAxis - EAxis::X); // Ensure that X axis is represented by bit 0.
+}
+
+// class FScopedBulkSelection //
+
 FScopedBulkSelection::FScopedBulkSelection(TSharedPtr<FPhysicsAssetEditorSharedData> InSharedData)
 	: SharedData(InSharedData)
 {
@@ -103,28 +113,6 @@ FScopedBulkSelection::~FScopedBulkSelection()
 	SharedData->BroadcastSelectionChanged();
 }
 
-// struct FBodyData //
-
-FBodyData::FBodyData()
-: CoMAxisFixedInComponentSpaceFlags(0)
-{}
-
-bool FBodyData::IsCoMAxisFixedInComponentSpace(const EAxis::Type InAxis) const
-{
-	return CoMAxisFixedInComponentSpaceFlags & GetBitFlag(InAxis);
-}
-
-void FBodyData::SetCoMAxisFixedInComponentSpace(const EAxis::Type InAxis, const bool bValue)
-{
-	CoMAxisFixedInComponentSpaceFlags = (bValue) ? CoMAxisFixedInComponentSpaceFlags | GetBitFlag(InAxis) : CoMAxisFixedInComponentSpaceFlags & ~GetBitFlag(InAxis);
-}
-
-int32 FBodyData::GetBitFlag(const EAxis::Type InAxis) const
-{
-	return int32(1) << InAxis;
-}
-
-
 // class FPhysicsAssetEditorSharedData //
 
 FPhysicsAssetEditorSharedData::FPhysicsAssetEditorSharedData()
@@ -244,7 +232,7 @@ void FPhysicsAssetEditorSharedData::Initialize(const TSharedRef<IPersonaPreviewS
 		}
 	}
 
-	EditorBodyData.SetNum(PhysicsAsset->SkeletalBodySetups.Num(), EAllowShrinking::Yes);
+	PhysicsAsset->EditorBodyFlags.SetNum(PhysicsAsset->SkeletalBodySetups.Num(), EAllowShrinking::Yes);
 
 	// Support undo/redo
 	PhysicsAsset->SetFlags(RF_Transactional);
@@ -399,9 +387,9 @@ FVector FPhysicsAssetEditorSharedData::GetCOMRenderPosition(const int32 BodyInde
 
 bool FPhysicsAssetEditorSharedData::IsCoMAxisFixedInComponentSpace(const int32 BodyIndex, const EAxis::Type InAxis) const
 {
-	if (EditorBodyData.IsValidIndex(BodyIndex))
+	if (PhysicsAsset && PhysicsAsset->EditorBodyFlags.IsValidIndex(BodyIndex))
 	{
-		return EditorBodyData[BodyIndex].IsCoMAxisFixedInComponentSpace(InAxis);
+		return PhysicsAsset->EditorBodyFlags[BodyIndex] & FindCoMAxisEditorBodyFlag(InAxis);
 	}
 
 	return false;
@@ -409,9 +397,10 @@ bool FPhysicsAssetEditorSharedData::IsCoMAxisFixedInComponentSpace(const int32 B
 
 void FPhysicsAssetEditorSharedData::SetCoMAxisFixedInComponentSpace(const int32 BodyIndex, const EAxis::Type InAxis, const bool bValue)
 {
-	if (EditorBodyData.IsValidIndex(BodyIndex))
+	if (PhysicsAsset && PhysicsAsset->EditorBodyFlags.IsValidIndex(BodyIndex))
 	{
-		return EditorBodyData[BodyIndex].SetCoMAxisFixedInComponentSpace(InAxis, bValue);
+		int32& BodyFlags = PhysicsAsset->EditorBodyFlags[BodyIndex];
+		BodyFlags = (bValue) ? BodyFlags | FindCoMAxisEditorBodyFlag(InAxis) : BodyFlags & ~FindCoMAxisEditorBodyFlag(InAxis);
 	}
 }
 
diff --git a/Engine/Source/Editor/PhysicsAssetEditor/Private/PhysicsAssetEditorSharedData.h b/Engine/Source/Editor/PhysicsAssetEditor/Private/PhysicsAssetEditorSharedData.h
--- a/Engine/Source/Editor/PhysicsAssetEditor/Private/PhysicsAssetEditorSharedData.h
+++ b/Engine/Source/Editor/PhysicsAssetEditor/Private/PhysicsAssetEditorSharedData.h
@@ -30,25 +30,6 @@ struct FScopedBulkSelection
 	TSharedPtr<FPhysicsAssetEditorSharedData> SharedData;
 };
 
-// BodyData
-//
-// A structure to keep track of editor only per physics body settings.
-//
-struct FBodyData
-{
-public:
-	FBodyData();
-
-	bool IsCoMAxisFixedInComponentSpace(const EAxis::Type InAxis) const;
-	void SetCoMAxisFixedInComponentSpace(const EAxis::Type InAxis, const bool bValue);
-
-private:
-	int32 GetBitFlag(const EAxis::Type InAxis) const;
-
-	int32 CoMAxisFixedInComponentSpaceFlags;
-};
-
-
 /*-----------------------------------------------------------------------------
    FPhysicsAssetEditorSharedData
 -----------------------------------------------------------------------------*/
@@ -391,9 +372,6 @@ public:
 	TArray<FBoneVertInfo> DominantWeightBoneInfos;
 	TArray<FBoneVertInfo> AnyWeightBoneInfos;
 
-	/** Editor only data for each physics body. */
-	TArray<FBodyData> EditorBodyData;
-
 	TArray<FSelection> SelectedBodies;
 
 	FSelection * GetSelectedBody()
diff --git a/Engine/Source/Runtime/Engine/Classes/PhysicsEngine/PhysicsAsset.h b/Engine/Source/Runtime/Engine/Classes/PhysicsEngine/PhysicsAsset.h
--- a/Engine/Source/Runtime/Engine/Classes/PhysicsEngine/PhysicsAsset.h
+++ b/Engine/Source/Runtime/Engine/Classes/PhysicsEngine/PhysicsAsset.h
@@ -188,6 +188,10 @@ class UPhysicsAsset : public UObject, public IInterface_PreviewMeshProvider
 	UPROPERTY(EditAnywhere, Category = Profiles, meta=(DisableCopyPaste))
 	TArray<FName> ConstraintProfiles;
 
+	/** A set of flags for each physics body. Used to store persistent per body editor data. */
+	UPROPERTY()
+	TArray<int32> EditorBodyFlags;
+
 	UPROPERTY(transient)
 	FName CurrentPhysicalAnimationProfileName;
 
