[UE] Phat: Minor UI bug fixes
- Hide CoM markers for hidden bodies
- Update Skel Tree icons after Simulate Selected

diff --git a/Engine/Source/Editor/Persona/Private/PhysicsAssetRenderUtils.cpp b/Engine/Source/Editor/Persona/Private/PhysicsAssetRenderUtils.cpp
--- a/Engine/Source/Editor/Persona/Private/PhysicsAssetRenderUtils.cpp
+++ b/Engine/Source/Editor/Persona/Private/PhysicsAssetRenderUtils.cpp
@@ -456,7 +456,7 @@ namespace PhysicsAssetRender
 		}
 	}
 
-	void DebugDrawCenterOfMass(USkeletalMeshComponent* const SkeletalMeshComponent, UPhysicsAsset* const PhysicsAsset, FPrimitiveDrawInterface* PDI, TFunctionRef< FVector(const uint32) > GetCoMPosition, TFunctionRef< bool(const uint32) > IsSelected, CreateCoMHitProxyFn CreateHitProxy)
+	void DebugDrawCenterOfMass(USkeletalMeshComponent* const SkeletalMeshComponent, UPhysicsAsset* const PhysicsAsset, FPrimitiveDrawInterface* PDI, TFunctionRef< FVector(const uint32) > GetCoMPosition, TFunctionRef< bool(const uint32) > IsSelected, TFunctionRef< bool(const uint32) > IsHidden, CreateCoMHitProxyFn CreateHitProxy)
 	{
 		static const float SelectedItemRenderSizeMultiplier = 1.5f;
 
@@ -470,7 +470,7 @@ namespace PhysicsAssetRender
 			{
 				const bool bIsSelected = IsSelected(BodyIndex);
 
-				if (!bDrawSelectedOnly || bIsSelected)
+				if (!IsHidden(BodyIndex) && (!bDrawSelectedOnly || bIsSelected))
 				{
 					const FBodyInstance* const BodyInstance = SkeletalMeshComponent->Bodies[BodyIndex];
 
diff --git a/Engine/Source/Editor/Persona/Public/PhysicsAssetRenderUtils.h b/Engine/Source/Editor/Persona/Public/PhysicsAssetRenderUtils.h
--- a/Engine/Source/Editor/Persona/Public/PhysicsAssetRenderUtils.h
+++ b/Engine/Source/Editor/Persona/Public/PhysicsAssetRenderUtils.h
@@ -226,7 +226,7 @@ namespace PhysicsAssetRender
 	PERSONA_API void DebugDrawBodies(class USkeletalMeshComponent* const SkeletalMeshComponent, class UPhysicsAsset* const PhysicsAsset, FPrimitiveDrawInterface* PDI, GetPrimitiveRef< FColor > GetPrimitiveColor, GetPrimitiveRef< class UMaterialInterface* > GetPrimitiveMaterial, GetPrimitiveTransformRef GetPrimitiveTransform, CreateBodyHitProxyFn CreateHitProxy);
 
 	/** Debug draw the Center of Masses for all bodies in the supplied Skeletal Mesh */
-	PERSONA_API void DebugDrawCenterOfMass(USkeletalMeshComponent* const SkeletalMeshComponent, class UPhysicsAsset* const PhysicsAsset, FPrimitiveDrawInterface* PDI, TFunctionRef< FVector(const uint32) > GetCoMPosition, TFunctionRef< bool(const uint32) > IsSelected, CreateCoMHitProxyFn CreateHitProxy);
+	PERSONA_API void DebugDrawCenterOfMass(USkeletalMeshComponent* const SkeletalMeshComponent, class UPhysicsAsset* const PhysicsAsset, FPrimitiveDrawInterface* PDI, TFunctionRef< FVector(const uint32) > GetCoMPosition, TFunctionRef< bool(const uint32) > IsSelected, TFunctionRef< bool(const uint32) > IsHidden, CreateCoMHitProxyFn CreateHitProxy);
 
 	/** Debug draw Physics Asset constraints using the supplied custom callbacks */
 	PERSONA_API void DebugDrawConstraints(class USkeletalMeshComponent* const SkeletalMeshComponent, class UPhysicsAsset* const PhysicsAsset, FPrimitiveDrawInterface* PDI, IsSelectedFn IsSelected, const bool bRunningSimulation, CreateConstraintHitProxyFn CreateHitProxy);
diff --git a/Engine/Source/Editor/PhysicsAssetEditor/Private/PhysicsAssetEditorEditMode.cpp b/Engine/Source/Editor/PhysicsAssetEditor/Private/PhysicsAssetEditorEditMode.cpp
--- a/Engine/Source/Editor/PhysicsAssetEditor/Private/PhysicsAssetEditorEditMode.cpp
+++ b/Engine/Source/Editor/PhysicsAssetEditor/Private/PhysicsAssetEditorEditMode.cpp
@@ -707,7 +707,7 @@ void FPhysicsAssetEditorEditMode::DrawHUD(FEditorViewportClient* ViewportClient,
 			// iterate over each body
 			for (int32 BodyIndex = 0, BodyCount = SharedData->EditorSkelComp->Bodies.Num(); BodyIndex < BodyCount; ++BodyIndex)
 			{
-				if ((!bDrawSelectedOnly || SharedData->IsBodySelected(BodyIndex)) && !(RenderSettings->bHideCenterOfMassForKinematicBodies && IsBodyKinematic(SharedData->PhysicsAsset.Get(), BodyIndex)))
+				if (!SharedData->IsBodyHidden(BodyIndex) && (!bDrawSelectedOnly || SharedData->IsBodySelected(BodyIndex)) && !(RenderSettings->bHideCenterOfMassForKinematicBodies && IsBodyKinematic(SharedData->PhysicsAsset.Get(), BodyIndex)))
 				{
 					FBodyInstance* BodyInst = SharedData->EditorSkelComp->Bodies[BodyIndex];
 					check(BodyInst);
diff --git a/Engine/Source/Editor/PhysicsAssetEditor/Private/PhysicsAssetEditorSharedData.cpp b/Engine/Source/Editor/PhysicsAssetEditor/Private/PhysicsAssetEditorSharedData.cpp
--- a/Engine/Source/Editor/PhysicsAssetEditor/Private/PhysicsAssetEditorSharedData.cpp
+++ b/Engine/Source/Editor/PhysicsAssetEditor/Private/PhysicsAssetEditorSharedData.cpp
@@ -3067,6 +3067,7 @@ void FPhysicsAssetEditorSharedData::EnableSimulation(bool bEnableSimulation)
 			EditorSkelComp->EnablePreview(true, PreviewAnimationAsset);
 		}
 
+		BroadcastHierarchyChanged();
 		BroadcastPreviewChanged();
 	}
 
diff --git a/Engine/Source/Editor/PhysicsAssetEditor/Private/PhysicsAssetEditorSkeletalMeshComponent.cpp b/Engine/Source/Editor/PhysicsAssetEditor/Private/PhysicsAssetEditorSkeletalMeshComponent.cpp
--- a/Engine/Source/Editor/PhysicsAssetEditor/Private/PhysicsAssetEditorSkeletalMeshComponent.cpp
+++ b/Engine/Source/Editor/PhysicsAssetEditor/Private/PhysicsAssetEditorSkeletalMeshComponent.cpp
@@ -135,25 +135,64 @@ void UPhysicsAssetEditorSkeletalMeshComponent::DebugDraw(const FSceneView* View,
 
 		// Draw Bodies.
 		{
-			auto TransformFn = [this](const UPhysicsAsset* PhysicsAsset, const FTransform& BoneTM, const int32 BodyIndex, const EAggCollisionShape::Type PrimType, const int32 PrimIndex, const float Scale) { return this->GetPrimitiveTransform(BoneTM, BodyIndex, PrimType, PrimIndex, Scale);  };
-			auto ColorFn = [this](const int32 BodyIndex, const EAggCollisionShape::Type PrimitiveType, const int32 PrimitiveIndex, const FPhysicsAssetRenderSettings& Settings) { return this->GetPrimitiveColor(BodyIndex, PrimitiveType, PrimitiveIndex); };
-			auto MaterialFn = [this](const int32 BodyIndex, const EAggCollisionShape::Type PrimitiveType, const int32 PrimitiveIndex, const FPhysicsAssetRenderSettings& Settings) { return this->GetPrimitiveMaterial(BodyIndex, PrimitiveType, PrimitiveIndex); };
-			auto HitProxyFn = [](const int32 BodyIndex, const EAggCollisionShape::Type PrimitiveType, const int32 PrimitiveIndex) { return new HPhysicsAssetEditorEdBoneProxy(BodyIndex, PrimitiveType, PrimitiveIndex); };
+			auto TransformFn = [this](const UPhysicsAsset* PhysicsAsset, const FTransform& BoneTM, const int32 BodyIndex, const EAggCollisionShape::Type PrimType, const int32 PrimIndex, const float Scale)
+				{
+					return this->GetPrimitiveTransform(BoneTM, BodyIndex, PrimType, PrimIndex, Scale);
+				};
+
+			auto ColorFn = [this](const int32 BodyIndex, const EAggCollisionShape::Type PrimitiveType, const int32 PrimitiveIndex, const FPhysicsAssetRenderSettings& Settings)
+				{
+					return this->GetPrimitiveColor(BodyIndex, PrimitiveType, PrimitiveIndex);
+				};
+
+			auto MaterialFn = [this](const int32 BodyIndex, const EAggCollisionShape::Type PrimitiveType, const int32 PrimitiveIndex, const FPhysicsAssetRenderSettings& Settings)
+				{
+					return this->GetPrimitiveMaterial(BodyIndex, PrimitiveType, PrimitiveIndex);
+				};
+
+			auto HitProxyFn = [](const int32 BodyIndex, const EAggCollisionShape::Type PrimitiveType, const int32 PrimitiveIndex)
+				{
+					return new HPhysicsAssetEditorEdBoneProxy(BodyIndex, PrimitiveType, PrimitiveIndex);
+				};
 
 			PhysicsAssetRender::DebugDrawBodies(this, PhysicsAsset, PDI, ColorFn, MaterialFn, TransformFn, HitProxyFn);
 		}
 		
 		{
-			auto COMPositionFn = [this](const int32 BodyIndex) { return this->SharedData->GetCOMRenderPosition(BodyIndex);  };
-			auto IsSelectedFn = [this](const uint32 InIndex) { return this->SharedData->IsBodySelected(InIndex) || this->SharedData->IsCoMSelected(InIndex); };
-			auto HitProxyFn = [](const int32 BodyIndex) { return new HPhysicsAssetEditorEdCoMProxy(BodyIndex); };
-			PhysicsAssetRender::DebugDrawCenterOfMass(this, PhysicsAsset, PDI, COMPositionFn, IsSelectedFn, HitProxyFn);
+			auto COMPositionFn = [this](const int32 BodyIndex) 
+				{ 
+					return this->SharedData->GetCOMRenderPosition(BodyIndex);  
+				};
+
+			auto IsSelectedFn = [this](const uint32 InIndex)
+				{
+					return this->SharedData->IsBodySelected(InIndex) || this->SharedData->IsCoMSelected(InIndex); 
+				};
+
+			auto IsHiddenFn = [this](const int32 BodyIndex)
+				{
+					return this->SharedData->IsBodyHidden(BodyIndex); 
+				};
+
+			auto HitProxyFn = [](const int32 BodyIndex)
+				{
+					return new HPhysicsAssetEditorEdCoMProxy(BodyIndex);
+				};
+
+			PhysicsAssetRender::DebugDrawCenterOfMass(this, PhysicsAsset, PDI, COMPositionFn, IsSelectedFn, IsHiddenFn, HitProxyFn);
 		}
 
 		// Draw Constraints.
 		{
-			auto HitProxyFn = [](const int32 InConstraintIndex) { return new HPhysicsAssetEditorEdConstraintProxy(InConstraintIndex); };
-			auto IsConstraintSelectedFn = [this](const uint32 InConstraintIndex) { return this->SharedData->IsConstraintSelected(InConstraintIndex); };
+			auto HitProxyFn = [](const int32 InConstraintIndex)
+				{
+					return new HPhysicsAssetEditorEdConstraintProxy(InConstraintIndex);
+				};
+
+			auto IsConstraintSelectedFn = [this](const uint32 InConstraintIndex)
+				{
+					return this->SharedData->IsConstraintSelected(InConstraintIndex);
+				};
 
 			PhysicsAssetRender::DebugDrawConstraints(this, PhysicsAsset, PDI, IsConstraintSelectedFn, SharedData->bRunningSimulation, HitProxyFn);
 		}
