[UE] Phat: Atempting to delete a Body selected in the skeleton tree will now delete that body and all its primitives, rather than just deleteing one of its primitives.

diff --git a/Engine/Source/Editor/PhysicsAssetEditor/Private/PhysicsAssetEditor.cpp b/Engine/Source/Editor/PhysicsAssetEditor/Private/PhysicsAssetEditor.cpp
--- a/Engine/Source/Editor/PhysicsAssetEditor/Private/PhysicsAssetEditor.cpp
+++ b/Engine/Source/Editor/PhysicsAssetEditor/Private/PhysicsAssetEditor.cpp
@@ -3435,8 +3435,8 @@ void FPhysicsAssetEditor::OnDeleteAllBodiesBelow()
 
 void FPhysicsAssetEditor::OnDeleteSelection()
 {
-	SharedData->DeleteCurrentPrim();
-	SharedData->DeleteCurrentConstraint();
+	SharedData->DeleteCurrentSelection();
+
 	RecreatePhysicsState();
 }
 
@@ -4338,21 +4338,21 @@ void FPhysicsAssetEditor::HandleSelectionChanged(const TArrayView<TSharedPtr<ISk
 			SharedData->ClearSelectedConstraints();
 
 			bool bBoneSelected = false;
-			TArray<FPhysicsAssetEditorSharedData::FSelection> SelectedBodies;
-			TArray<int32> SelectedBodiesAnyPrim;
+			TArray<FPhysicsAssetEditorSharedData::FSelection> SelectedPrimitives;
+			TArray<int32> SelectedBodies;
 			TArray<int32> SelectedConstraints;
 			for (const TSharedPtr<ISkeletonTreeItem>& Item : InSelectedItems)
 			{
 				if (Item->IsOfType<FSkeletonTreePhysicsBodyItem>())
 				{
 					TSharedPtr<FSkeletonTreePhysicsBodyItem> SkeletonTreePhysicsBodyItem = StaticCastSharedPtr<FSkeletonTreePhysicsBodyItem>(Item);
-					SelectedBodiesAnyPrim.Add(SkeletonTreePhysicsBodyItem->GetBodySetupIndex());
+					SelectedBodies.Add(SkeletonTreePhysicsBodyItem->GetBodySetupIndex());
 				}
 				else if (Item->IsOfType<FSkeletonTreePhysicsShapeItem>())
 				{
 					TSharedPtr<FSkeletonTreePhysicsShapeItem> SkeletonTreePhysicsShapeItem = StaticCastSharedPtr<FSkeletonTreePhysicsShapeItem>(Item);
 					FPhysicsAssetEditorSharedData::FSelection Selection(SkeletonTreePhysicsShapeItem->GetBodySetupIndex(), SkeletonTreePhysicsShapeItem->GetShapeType(), SkeletonTreePhysicsShapeItem->GetShapeIndex());
-					SelectedBodies.Add(Selection);
+					SelectedPrimitives.Add(Selection);
 
 				}
 				else if (Item->IsOfType<FSkeletonTreePhysicsConstraintItem>())
@@ -4366,8 +4366,8 @@ void FPhysicsAssetEditor::HandleSelectionChanged(const TArrayView<TSharedPtr<ISk
 				}
 			}
 
-			SharedData->SetSelectedBodies(SelectedBodies, true);
-			SharedData->SetSelectedBodiesAnyPrimitive(SelectedBodiesAnyPrim, true);
+			SharedData->SetSelectedBodies(SelectedPrimitives, true);
+			SharedData->SetSelectedBodiesFromDirectlInput(SelectedBodies, true);
 			SharedData->SetSelectedConstraints(SelectedConstraints, true);
 
 			if(!bBoneSelected)
diff --git a/Engine/Source/Editor/PhysicsAssetEditor/Private/PhysicsAssetEditorSharedData.cpp b/Engine/Source/Editor/PhysicsAssetEditor/Private/PhysicsAssetEditorSharedData.cpp
--- a/Engine/Source/Editor/PhysicsAssetEditor/Private/PhysicsAssetEditorSharedData.cpp
+++ b/Engine/Source/Editor/PhysicsAssetEditor/Private/PhysicsAssetEditorSharedData.cpp
@@ -99,6 +99,42 @@ int32 FindCoMAxisEditorBodyFlag(const EAxis::Type InAxis)
 	return int32(1) << (InAxis - EAxis::X); // Ensure that X axis is represented by bit 0.
 }
 
+namespace
+{
+	template <typename TShapeElem>
+	void SetSelectedBodiesPrimitivesHelper(const int32 BodyIndex, const TArray<TShapeElem>& ShapeElems, TArray<FPhysicsAssetEditorSharedData::FSelection>& SelectedElems, const TFunction<bool(const TArray<FPhysicsAssetEditorSharedData::FSelection>&, const int32 BodyIndex, const FKShapeElem&)>& Predicate)
+	{
+		for (int32 PrimitiveIndex = 0; PrimitiveIndex < ShapeElems.Num(); ++PrimitiveIndex)
+		{
+			const TShapeElem& ShapeElem = ShapeElems[PrimitiveIndex];
+			if (Predicate(SelectedElems, BodyIndex, ShapeElem))
+			{
+				SelectedElems.Add(FPhysicsAssetEditorSharedData::FSelection(BodyIndex, ShapeElem.GetShapeType(), PrimitiveIndex));
+			}
+		}
+	}
+}
+
+TArray<FPhysicsAssetEditorSharedData::FSelection> CreateBodyPrimitivesSelection(TObjectPtr<UPhysicsAsset> PhysicsAsset, const TArray<int32>& BodiesIndices, const TFunction<bool(const TArray<FPhysicsAssetEditorSharedData::FSelection>&, const int32 BodyIndex, const FKShapeElem&)>& Predicate)
+{
+	TArray<FPhysicsAssetEditorSharedData::FSelection> NewSelection;
+	for (const int32 BodyIndex : BodiesIndices)
+	{
+		UBodySetup* BodySetup = PhysicsAsset->SkeletalBodySetups[BodyIndex];
+		check(BodySetup);
+
+		const FKAggregateGeom& AggGeom = BodySetup->AggGeom;
+		SetSelectedBodiesPrimitivesHelper(BodyIndex, AggGeom.SphereElems, NewSelection, Predicate);
+		SetSelectedBodiesPrimitivesHelper(BodyIndex, AggGeom.BoxElems, NewSelection, Predicate);
+		SetSelectedBodiesPrimitivesHelper(BodyIndex, AggGeom.SphylElems, NewSelection, Predicate);
+		SetSelectedBodiesPrimitivesHelper(BodyIndex, AggGeom.ConvexElems, NewSelection, Predicate);
+		SetSelectedBodiesPrimitivesHelper(BodyIndex, AggGeom.TaperedCapsuleElems, NewSelection, Predicate);
+	}
+
+	return NewSelection;
+}
+
+
 // class FScopedBulkSelection //
 
 FScopedBulkSelection::FScopedBulkSelection(TSharedPtr<FPhysicsAssetEditorSharedData> InSharedData)
@@ -856,7 +892,7 @@ void FPhysicsAssetEditorSharedData::SetSelectedBodyAnyPrimitive(int32 BodyIndex,
 	SetSelectedBodiesAnyPrimitive({ BodyIndex }, bSelected);
 }
 
-void FPhysicsAssetEditorSharedData::SetSelectedBodiesAnyPrimitive(const TArray<int32>& BodiesIndices, bool bSelected)
+void FPhysicsAssetEditorSharedData::SetSelectedBodiesAnyPrimitive(const TArray<int32>& BodiesIndices, const bool bSelected)
 {
 	SetSelectedBodiesPrimitives(BodiesIndices, bSelected, [](const TArray<FSelection>& CurrentSelection, const int32 BodyIndex, const FKShapeElem& Primitive)
 	{
@@ -874,6 +910,33 @@ void FPhysicsAssetEditorSharedData::SetSelectedBodiesAllPrimitive(const TArray<i
 	});
 }
 
+void FPhysicsAssetEditorSharedData::SetSelectedBodiesFromDirectlInput(const TArray<int32>& BodiesIndices, const bool bSelected)
+{
+	if (BodiesIndices.Num() == 0)
+	{
+		return;
+	}
+
+	if (BodiesIndices.Num() == 1 && BodiesIndices[0] == INDEX_NONE)
+	{
+		ClearSelectedBody();
+		return;
+	}
+
+	// Create a selection for each body that references the first primitive in that body. (ordering of primitives is determined by geometry type then by order if instantiation).
+	TArray<FSelection> SelectionCollection = CreateBodyPrimitivesSelection(PhysicsAsset, BodiesIndices, [this](const TArray<FSelection>& CurrentSelection, const int32 BodyIndex, const FKShapeElem& Primitive) 
+		{ 
+			return !IsBodySelected(BodyIndex); 
+		});
+
+	for (FSelection& Selection : SelectionCollection)
+	{
+		Selection.bIsDirectBodySelection = true;
+	}
+
+	SetSelectedBodies(SelectionCollection, bSelected);
+}
+
 void FPhysicsAssetEditorSharedData::SetSelectedBodiesPrimitivesWithCollisionType(const TArray<int32>& BodiesIndices, const ECollisionEnabled::Type CollisionType, bool bSelected)
 {
 	SetSelectedBodiesPrimitives(BodiesIndices, bSelected, [CollisionType](const TArray<FSelection>& CurrentSelection, const int32 BodyIndex, const FKShapeElem& Primitive)
@@ -883,22 +946,6 @@ void FPhysicsAssetEditorSharedData::SetSelectedBodiesPrimitivesWithCollisionType
 	});
 }
 
-namespace
-{
-	template <typename TShapeElem>
-	void SetSelectedBodiesPrimitivesHelper(const int32 BodyIndex, const TArray<TShapeElem>& ShapeElems, TArray<FPhysicsAssetEditorSharedData::FSelection>& SelectedElems, const TFunction<bool(const TArray<FPhysicsAssetEditorSharedData::FSelection>&, const int32 BodyIndex, const FKShapeElem&)>& Predicate)
-	{
-		for (int32 PrimitiveIndex = 0; PrimitiveIndex < ShapeElems.Num(); ++PrimitiveIndex)
-		{
-			const TShapeElem& ShapeElem = ShapeElems[PrimitiveIndex];
-			if (Predicate(SelectedElems, BodyIndex, ShapeElem))
-			{
-				SelectedElems.Add(FPhysicsAssetEditorSharedData::FSelection(BodyIndex, ShapeElem.GetShapeType(), PrimitiveIndex));
-			}
-		}
-	}
-}
-
 void FPhysicsAssetEditorSharedData::SetSelectedBodiesPrimitives(const TArray<int32>& BodiesIndices, bool bSelected, const TFunction<bool(const TArray<FSelection>&, const int32 BodyIndex, const FKShapeElem&)>& Predicate)
 {
 	if (BodiesIndices.Num() == 0)
@@ -912,21 +959,7 @@ void FPhysicsAssetEditorSharedData::SetSelectedBodiesPrimitives(const TArray<int
 		return;
 	}
 
-	TArray<FSelection> NewSelection;
-	for (const int32 BodyIndex : BodiesIndices)
-	{
-		UBodySetup* BodySetup = PhysicsAsset->SkeletalBodySetups[BodyIndex];
-		check(BodySetup);
-
-		const FKAggregateGeom& AggGeom = BodySetup->AggGeom;
-		SetSelectedBodiesPrimitivesHelper(BodyIndex, AggGeom.SphereElems, NewSelection, Predicate);
-		SetSelectedBodiesPrimitivesHelper(BodyIndex, AggGeom.BoxElems, NewSelection, Predicate);
-		SetSelectedBodiesPrimitivesHelper(BodyIndex, AggGeom.SphylElems, NewSelection, Predicate);
-		SetSelectedBodiesPrimitivesHelper(BodyIndex, AggGeom.ConvexElems, NewSelection, Predicate);
-		SetSelectedBodiesPrimitivesHelper(BodyIndex, AggGeom.TaperedCapsuleElems, NewSelection, Predicate);
-	}
-
-	SetSelectedBodies(NewSelection, bSelected);
+	SetSelectedBodies(CreateBodyPrimitivesSelection(PhysicsAsset, BodiesIndices, Predicate), bSelected);
 }
 
 void FPhysicsAssetEditorSharedData::ClearSelected()
@@ -2003,7 +2036,6 @@ void FPhysicsAssetEditorSharedData::PasteBodiesAndConstraintsFromClipboard(int32
 		}
 	}
 }
-
 void FPhysicsAssetEditorSharedData::CopySelectedShapesToClipboard(int32& OutNumCopiedShapes, int32& OutNumBodiesCopiedFrom)
 {
 	OutNumCopiedShapes = 0;
@@ -2714,7 +2746,6 @@ bool FPhysicsAssetEditorSharedData::IsAngularConstraintLocked(EPhysicsAssetEdito
 	return bLocked;
 }
 
-
 void FPhysicsAssetEditorSharedData::DeleteBody(int32 DelBodyIndex, bool bRefreshComponent)
 {
 	USkeletalMesh* EditorSkelMesh = PhysicsAsset->GetPreviewMesh();
@@ -2791,6 +2822,34 @@ void FPhysicsAssetEditorSharedData::DeleteBody(int32 DelBodyIndex, bool bRefresh
 	}
 }
 
+void FPhysicsAssetEditorSharedData::DeleteCurrentSelection()
+{
+	DeleteCurrentBody();
+	DeleteCurrentPrim();
+	DeleteCurrentConstraint();
+}
+
+void FPhysicsAssetEditorSharedData::DeleteCurrentBody()
+{
+	// Delete any directly selected bodies and all their primitives
+
+	TArray<FSelection> SelectionsToRemove;
+
+	for (const FSelection& Selection : SelectedBodies)
+	{
+		if (Selection.bIsDirectBodySelection)
+		{
+			DeleteBody(Selection.Index, false);
+			SelectionsToRemove.Add(Selection);
+		}
+	}
+
+	SetSelectedBodies(SelectionsToRemove, false);
+
+	RefreshPhysicsAssetChange(PhysicsAsset);
+	BroadcastHierarchyChanged();
+}
+
 void FPhysicsAssetEditorSharedData::DeleteCurrentPrim()
 {
 	if (bRunningSimulation)
diff --git a/Engine/Source/Editor/PhysicsAssetEditor/Private/PhysicsAssetEditorSharedData.h b/Engine/Source/Editor/PhysicsAssetEditor/Private/PhysicsAssetEditorSharedData.h
--- a/Engine/Source/Editor/PhysicsAssetEditor/Private/PhysicsAssetEditorSharedData.h
+++ b/Engine/Source/Editor/PhysicsAssetEditor/Private/PhysicsAssetEditorSharedData.h
@@ -57,6 +57,7 @@ public:
 		FTransform WidgetTM;
 		FTransform ManipulateTM;
 		FVector CoMPosition; 
+		bool bIsDirectBodySelection; // When true, this selection represents a body that has been directly selected, not just one of its primitives.
 
 		FSelection(int32 GivenBodyIndex, EAggCollisionShape::Type GivenPrimitiveType, int32 GivenPrimitiveIndex) 
 			: Index(GivenBodyIndex)
@@ -64,12 +65,13 @@ public:
 			, PrimitiveIndex(GivenPrimitiveIndex)
 			, WidgetTM(FTransform::Identity)
 			, ManipulateTM(FTransform::Identity)
+			, bIsDirectBodySelection(false)
 		{
 		}
 
 		bool operator==(const FSelection& rhs) const
 		{
-			return Index == rhs.Index && PrimitiveType == rhs.PrimitiveType && PrimitiveIndex == rhs.PrimitiveIndex;
+			return Index == rhs.Index && PrimitiveType == rhs.PrimitiveType && PrimitiveIndex == rhs.PrimitiveIndex && bIsDirectBodySelection == rhs.bIsDirectBodySelection;
 		}
 	};
 
@@ -128,9 +130,6 @@ public:
 	/** Snaps the specified constraint to it's bone */
 	void SnapConstraintToBone(FConstraintInstance& ConstraintInstance, const EConstraintTransformComponentFlags ComponentFlags = EConstraintTransformComponentFlags::All);
 
-	/** Deletes the currently selected constraints */
-	void DeleteCurrentConstraint();
-
 	/** Paste the currently-copied constraint properties onto the single selected constraint */
 	void PasteConstraintProperties();
 	
@@ -147,6 +146,21 @@ public:
 	/** Gets whether the active constraint is locked */
 	bool IsAngularConstraintLocked(EPhysicsAssetEditorConstraintType Constraint) const;
 
+	void DeleteBody(int32 DelBodyIndex, bool bRefreshComponent = true);
+
+	/** Deletes all currently selected objects */
+	void DeleteCurrentSelection();
+
+	/** Deletes the currently selected bodies and all their primitives */
+	void DeleteCurrentBody();
+
+	/** Deletes the currently selected primitives */
+	void DeleteCurrentPrim();
+
+	/** Deletes the currently selected constraints */
+	void DeleteCurrentConstraint();
+
+
 	/** Center of Mass editing */
 	void ClearSelectedCoMs();
 	void SetSelectedCoM(const FSelection& InSelectedCoM, const bool bSelected);
@@ -179,10 +193,9 @@ public:
 	void SetSelectedBodyAnyPrimitive(int32 BodyIndex, bool bSelected);
 	void SetSelectedBodiesAnyPrimitive(const TArray<int32>& BodiesIndices, bool bSelected);
 	void SetSelectedBodiesAllPrimitive(const TArray<int32>& BodiesIndices, bool bSelected);
+	void SetSelectedBodiesFromDirectlInput(const TArray<int32>& BodiesIndices, const bool bSelected);
 	void SetSelectedBodiesPrimitivesWithCollisionType(const TArray<int32>& BodiesIndices, const ECollisionEnabled::Type CollisionType, bool bSelected);
 	void SetSelectedBodiesPrimitives(const TArray<int32>& BodiesIndices, bool bSelected, const TFunction<bool(const TArray<FSelection>&, const int32 BodyIndex, const FKShapeElem&)>& Predicate);
-	void DeleteCurrentPrim();
-	void DeleteBody(int32 DelBodyIndex, bool bRefreshComponent=true);
 	void RefreshPhysicsAssetChange(const UPhysicsAsset* InPhysAsset, bool bFullClothRefresh = true);
 	void MakeNewBody(int32 NewBoneIndex, bool bAutoSelect = true);
 	void MakeNewBody(const FPhysAssetCreateParams& NewBodyData, const int32 NewBoneIndex, const bool bAutoSelect = true);
