[UE] Phat: Ctrl+Shift+C/P to copy and paste bodies will now copy/paste Collision relationships as well.

diff --git a/Engine/Source/Editor/PhysicsAssetEditor/Private/PhysicsAssetEditor.cpp b/Engine/Source/Editor/PhysicsAssetEditor/Private/PhysicsAssetEditor.cpp
--- a/Engine/Source/Editor/PhysicsAssetEditor/Private/PhysicsAssetEditor.cpp
+++ b/Engine/Source/Editor/PhysicsAssetEditor/Private/PhysicsAssetEditor.cpp
@@ -2686,10 +2686,11 @@ void FPhysicsAssetEditor::OnCopyBodies()
 {
 	int32 NumCopiedBodies;
 	int32 NumCopiedConstraints;
-	SharedData->CopySelectedBodiesAndConstraintsToClipboard(NumCopiedBodies, NumCopiedConstraints);
+	int32 NumCopiedDisabledCollisionPairs;
+	SharedData->CopySelectedBodiesAndConstraintsToClipboard(NumCopiedBodies, NumCopiedConstraints, NumCopiedDisabledCollisionPairs);
 
-	const FText MessageFormat = LOCTEXT("CopiedBodiesAndConstraintsToClipboard", "{0} bodies and {1} constraints copied to clipboard");
-	ShowNotificationMessage(FText::Format(MessageFormat, NumCopiedBodies, NumCopiedConstraints), SNotificationItem::CS_Success);
+	const FText MessageFormat = LOCTEXT("CopiedBodiesAndConstraintsToClipboard", "{0} {0}|plural(one=body,other=bodies), {1} {1}|plural(one=constraints,other=constraints) and {2} disabled collision {2}|plural(one=pair,other=pairs) copied to clipboard");
+	ShowNotificationMessage(FText::Format(MessageFormat, NumCopiedBodies, NumCopiedConstraints, NumCopiedDisabledCollisionPairs), SNotificationItem::CS_Success);
 }
 
 bool FPhysicsAssetEditor::IsCopyBodies() const
@@ -2711,10 +2712,11 @@ void FPhysicsAssetEditor::OnPasteBodies()
 {
 	int32 NumPastedBodies;
 	int32 NumPastedConstraints;
-	SharedData->PasteBodiesAndConstraintsFromClipboard(NumPastedBodies, NumPastedConstraints);
+	int32 NumCopiedDisabledCollisionPairs;
+	SharedData->PasteBodiesAndConstraintsFromClipboard(NumPastedBodies, NumPastedConstraints, NumCopiedDisabledCollisionPairs);
 
-	const FText MessageFormat = LOCTEXT("PastedBodiesAndConstraintsToClipboard", "{0} bodies and {1} constraints pasted from clipboard");
-	ShowNotificationMessage(FText::Format(MessageFormat, NumPastedBodies, NumPastedConstraints), SNotificationItem::CS_Success);
+	const FText MessageFormat = LOCTEXT("PastedBodiesAndConstraintsToClipboard", "{0} {0}|plural(one=body,other=bodies), {1} {1}|plural(one=constraints,other=constraints) and {2} disabled collision {2}|plural(one=pair,other=pairs) pasted from clipboard");
+	ShowNotificationMessage(FText::Format(MessageFormat, NumPastedBodies, NumPastedConstraints, NumCopiedDisabledCollisionPairs), SNotificationItem::CS_Success);
 }
 
 bool FPhysicsAssetEditor::CanPasteBodies() const
diff --git a/Engine/Source/Editor/PhysicsAssetEditor/Private/PhysicsAssetEditorSharedData.cpp b/Engine/Source/Editor/PhysicsAssetEditor/Private/PhysicsAssetEditorSharedData.cpp
--- a/Engine/Source/Editor/PhysicsAssetEditor/Private/PhysicsAssetEditorSharedData.cpp
+++ b/Engine/Source/Editor/PhysicsAssetEditor/Private/PhysicsAssetEditorSharedData.cpp
@@ -134,6 +134,22 @@ TArray<FPhysicsAssetEditorSharedData::FSelection> CreateBodyPrimitivesSelection(
 	return NewSelection;
 }
 
+// Pass each unique pair of values (excluding those containing the same value twice) in the supplied collection to the supplied function object.
+template<typename CollectionType, typename FunctionObjectType>
+void ForEachUniquePair(CollectionType& Collection, const FunctionObjectType& FunctionObject)
+{
+	for (typename CollectionType::TConstIterator OuterItr = Collection.CreateConstIterator(); OuterItr; ++OuterItr)
+	{
+		typename CollectionType::TConstIterator InnerItr = OuterItr;
+		
+		++InnerItr;
+
+		for (; InnerItr; ++InnerItr)
+		{
+			FunctionObject(*OuterItr, *InnerItr);
+		}
+	}
+}
 
 // class FScopedBulkSelection //
 
@@ -1827,10 +1843,12 @@ void FPhysicsAssetEditorSharedData::AutoNamePrimitive(int32 BodyIndex, EAggColli
 	}
 }
 
-void FPhysicsAssetEditorSharedData::CopySelectedBodiesAndConstraintsToClipboard(int32& OutNumCopiedBodies, int32& OutNumCopiedConstraints)
+void FPhysicsAssetEditorSharedData::CopySelectedBodiesAndConstraintsToClipboard(int32& OutNumCopiedBodies, int32& OutNumCopiedConstraints, int32& OutNumCopiedDisabledCollisionPairs)
 {
 	OutNumCopiedBodies = 0;
 	OutNumCopiedConstraints = 0;
+	OutNumCopiedDisabledCollisionPairs = 0;
+
 	if (PhysicsAsset)
 	{
 		// Clear the mark state for saving.
@@ -1839,10 +1857,11 @@ void FPhysicsAssetEditorSharedData::CopySelectedBodiesAndConstraintsToClipboard(
 		FStringOutputDevice Archive;
 		const FExportObjectInnerContext Context;
 
+		TSet<int32> ExportedBodyIndices;
+
 		// export bodies first 
 		{
 			OutNumCopiedBodies = 0;
-			TSet<int32> ExportedBodyIndices;
 
 			// Export each of the selected nodes
 			for (const FSelection& SelectedBody : SelectedBodies)
@@ -1870,7 +1889,7 @@ void FPhysicsAssetEditorSharedData::CopySelectedBodiesAndConstraintsToClipboard(
 			// Export each of the selected nodes
 			for (const FSelection& SelectedConstraint : SelectedConstraints)
 			{
-				// selected bodies contain the primitives, so abody can be stored multiple time for each of its primitive
+				// selected bodies contain the primitives, so a body can be stored multiple time for each of its primitive
 				// we need to make sure we process it only once
 				if (!ExportedConstraintIndices.Contains(SelectedConstraint.Index))
 				{
@@ -1885,6 +1904,32 @@ void FPhysicsAssetEditorSharedData::CopySelectedBodiesAndConstraintsToClipboard(
 			}
 		}
 
+		// export collision relationships
+		{
+			auto ExportCollisionRelationships = [this, &Context, &Archive, &OutNumCopiedDisabledCollisionPairs](const int32 BodyIndexA, const int32 BodyIndexB)
+				{
+					if (!PhysicsAsset->IsCollisionEnabled(BodyIndexA, BodyIndexB))
+					{
+						const USkeletalBodySetup* const BodySetupA = PhysicsAsset->SkeletalBodySetups[BodyIndexA];
+						const USkeletalBodySetup* const BodySetupB = PhysicsAsset->SkeletalBodySetups[BodyIndexB];
+
+						check(BodySetupA);
+						check(BodySetupB);
+
+						if (BodySetupA && BodySetupB)
+						{
+							UPhysicsAssetCollisionPair* const CollisionPair = NewObject<UPhysicsAssetCollisionPair>();
+
+							CollisionPair->Set(BodySetupA->BoneName, BodySetupB->BoneName);
+							UExporter::ExportToOutputDevice(&Context, CollisionPair, NULL, Archive, TEXT("copy"), 0, PPF_ExportsNotFullyQualified | PPF_Copy | PPF_Delimited, false);
+							++OutNumCopiedDisabledCollisionPairs;
+						}
+					}
+				};
+
+			ForEachUniquePair(ExportedBodyIndices, ExportCollisionRelationships);
+		}
+
 		// save to clipboard as text 
 		FString ExportedText = Archive;
 		FPlatformApplicationMisc::ClipboardCopy(*ExportedText);
@@ -1902,7 +1947,7 @@ public:
 	// FCustomizableTextObjectFactory implementation
 	virtual bool CanCreateClass(UClass* InObjectClass, bool& bOmitSubObjs) const override
 	{
-		return (InObjectClass->IsChildOf<USkeletalBodySetup>() || InObjectClass->IsChildOf<UPhysicsConstraintTemplate>());
+		return (InObjectClass->IsChildOf<USkeletalBodySetup>() || InObjectClass->IsChildOf<UPhysicsConstraintTemplate>() || InObjectClass->IsChildOf<UPhysicsAssetCollisionPair>());
 	}
 
 	virtual void ProcessConstructedObject(UObject* NewObject) override
@@ -1916,11 +1961,16 @@ public:
 		{
 			NewConstraintTemplates.Add(Cast<UPhysicsConstraintTemplate>(NewObject));
 		}
+		else if (NewObject->IsA<UPhysicsAssetCollisionPair>())
+		{
+			NewDisabledCollisionPairs.Add(Cast<UPhysicsAssetCollisionPair>(NewObject));
+		}
 	}
 
 public:
 	TArray<USkeletalBodySetup*> NewBodySetups;
 	TArray<UPhysicsConstraintTemplate*> NewConstraintTemplates;
+	TArray<UPhysicsAssetCollisionPair*> NewDisabledCollisionPairs;
 };
 
 bool FPhysicsAssetEditorSharedData::CanPasteBodiesAndConstraintsFromClipboard() const
@@ -1931,10 +1981,12 @@ bool FPhysicsAssetEditorSharedData::CanPasteBodiesAndConstraintsFromClipboard()
 	return Factory.CanCreateObjectsFromText(TextToImport);
 }
 
-void FPhysicsAssetEditorSharedData::PasteBodiesAndConstraintsFromClipboard(int32& OutNumPastedBodies, int32& OutNumPastedConstraints)
+void FPhysicsAssetEditorSharedData::PasteBodiesAndConstraintsFromClipboard(int32& OutNumPastedBodies, int32& OutNumPastedConstraints, int32& OutNumPastedDisabledCollisionPairs)
 {
 	OutNumPastedBodies = 0;
 	OutNumPastedConstraints = 0;
+	OutNumPastedDisabledCollisionPairs = 0;
+
 	if (PhysicsAsset)
 	{
 		FString TextToImport;
@@ -1945,14 +1997,16 @@ void FPhysicsAssetEditorSharedData::PasteBodiesAndConstraintsFromClipboard(int32
 			UPackage* TempPackage = NewObject<UPackage>(nullptr, TEXT("/Engine/Editor/PhysicsAssetEditor/Transient"), RF_Transient);
 			TempPackage->AddToRoot();
 			{
+				TArray<int32> PastedBodyIndices;
+
 				// Turn the text buffer into objects
 				FSkeletalBodyAndConstraintSetupObjectTextFactory  Factory;
 				Factory.ProcessBuffer(TempPackage, RF_Transactional, TextToImport);
 
 				// transaction block 
-				if (Factory.NewBodySetups.Num() > 0 || Factory.NewConstraintTemplates.Num() > 0)
+				if (Factory.NewBodySetups.Num() > 0 || Factory.NewConstraintTemplates.Num() > 0 || Factory.NewDisabledCollisionPairs.Num() > 0)
 				{
-					const FScopedTransaction Transaction(NSLOCTEXT("PhysicsAssetEditor", "PasteBodiesAndConstraintsFromClipboard", "Paste Bodies And Constraints From Clipboard"));
+					const FScopedTransaction Transaction(NSLOCTEXT("PhysicsAssetEditor", "PasteBodiesAndConstraintsFromClipboard", "Paste Bodies, Constraints And Disabled Collision Pairs From Clipboard"));
 
 					PhysicsAsset->Modify();
 
@@ -1960,7 +2014,7 @@ void FPhysicsAssetEditorSharedData::PasteBodiesAndConstraintsFromClipboard(int32
 					OutNumPastedBodies = 0;
 					for (USkeletalBodySetup* PastedBodySetup : Factory.NewBodySetups)
 					{
-						// doe sthis bone exist in the target physics asset?
+						// does this bone exist in the target physics asset?
 						int32 BodyIndex = PhysicsAsset->FindBodyIndex(PastedBodySetup->BoneName);
 						if (BodyIndex == INDEX_NONE)
 						{
@@ -1977,6 +2031,8 @@ void FPhysicsAssetEditorSharedData::PasteBodiesAndConstraintsFromClipboard(int32
 								TargetBodySetup->Modify();
 								TargetBodySetup->CopyBodyPropertiesFrom(PastedBodySetup);
 								++OutNumPastedBodies;
+
+								PastedBodyIndices.Add(BodyIndex);
 							}
 						}
 					}
@@ -2023,6 +2079,25 @@ void FPhysicsAssetEditorSharedData::PasteBodiesAndConstraintsFromClipboard(int32
 							}
 						}
 					}
+
+					// Enable collisions between all pasted bodies.
+					ForEachUniquePair(PastedBodyIndices, [this](const int32 BodyIndexA, const int32 BodyIndexB) 
+						{ 
+							PhysicsAsset->EnableCollision(BodyIndexA, BodyIndexB); 
+						});
+
+					// Disable collisions between pasted bodies as specified by pasted disabled collision pairs.
+					for (UPhysicsAssetCollisionPair* const CollisionPair : Factory.NewDisabledCollisionPairs)
+					{
+						const int32 BodyIndexA = PhysicsAsset->FindBodyIndex(CollisionPair->BoneNameA);
+						const int32 BodyIndexB = PhysicsAsset->FindBodyIndex(CollisionPair->BoneNameB);
+
+						if ((BodyIndexA != INDEX_NONE) && (BodyIndexB != INDEX_NONE))
+						{
+							PhysicsAsset->DisableCollision(BodyIndexA, BodyIndexB);
+							++OutNumPastedDisabledCollisionPairs;
+						}
+					}
 				}
 			}
 			// Remove the temp package from the root now that it has served its purpose
diff --git a/Engine/Source/Editor/PhysicsAssetEditor/Private/PhysicsAssetEditorSharedData.h b/Engine/Source/Editor/PhysicsAssetEditor/Private/PhysicsAssetEditorSharedData.h
--- a/Engine/Source/Editor/PhysicsAssetEditor/Private/PhysicsAssetEditorSharedData.h
+++ b/Engine/Source/Editor/PhysicsAssetEditor/Private/PhysicsAssetEditorSharedData.h
@@ -10,6 +10,8 @@
 #include "PhysicsEngine/ShapeElem.h"
 #include "Preferences/PhysicsAssetEditorOptions.h"
 
+#include "PhysicsAssetEditorSharedData.generated.h"
+
 class UBodySetup;
 class UPhysicsAssetEditorSkeletalMeshComponent;
 class UPhysicsAsset;
@@ -30,6 +32,30 @@ struct FScopedBulkSelection
 	TSharedPtr<FPhysicsAssetEditorSharedData> SharedData;
 };
 
+/** 
+ * class UPhysicsAssetCollisionPair
+ * 
+ * Used to simplify Copy + Paste of collision relationships between physics bodies.
+ */
+UCLASS(MinimalAPI)
+class UPhysicsAssetCollisionPair : public UObject
+{
+	GENERATED_BODY()
+
+public:
+	void Set(const FName InBoneNameA, const FName InBoneNameB)
+	{
+		BoneNameA = InBoneNameA;
+		BoneNameB = InBoneNameB;
+	}
+
+	UPROPERTY()
+	FName BoneNameA;
+
+	UPROPERTY()
+	FName BoneNameB;
+};
+
 /*-----------------------------------------------------------------------------
    FPhysicsAssetEditorSharedData
 -----------------------------------------------------------------------------*/
@@ -201,9 +227,9 @@ public:
 	void MakeNewBody(const FPhysAssetCreateParams& NewBodyData, const int32 NewBoneIndex, const bool bAutoSelect = true);
 	void MakeNewConstraints(int32 ParentBodyIndex, const TArray<int32>& ChildBodyIndices);
 	void MakeNewConstraint(int32 ParentBodyIndex, int32 ChildBodyIndex);
-	void CopySelectedBodiesAndConstraintsToClipboard(int32& OutNumCopiedBodies, int32& OutNumCopiedConstraints);
+	void CopySelectedBodiesAndConstraintsToClipboard(int32& OutNumCopiedBodies, int32& OutNumCopiedConstraints, int32& OutNumCopiedDisabledCollisionPairs);
 	bool CanPasteBodiesAndConstraintsFromClipboard() const;
-	void PasteBodiesAndConstraintsFromClipboard(int32& OutNumPastedBodies, int32& OutNumPastedConstraints);
+	void PasteBodiesAndConstraintsFromClipboard(int32& OutNumPastedBodies, int32& OutNumPastedConstraints, int32& OutNumPastedDisabledCollisionPairs);
 	void CopySelectedShapesToClipboard(int32& OutNumCopiedShapes, int32& OutNumBodiesCopiedFrom);
 	bool CanPasteShapesFromClipboard() const;
 	void PasteShapesFromClipboard(int32& OutNumPastedShapes, int32& OutNumBodiesPastedInto);
