Add the "Physics Rendering" section to the "Show" menu of the new Animation viewport toolbar

* Wrap the code creating the "Physics Rendering" section in a lambda so we can call it twice, once for the old and once for the new viewport toolbar.
* Extend the new animation viewport toolbar context so it gets the required UPhysicsAssetEditorToolMenuContext object.

diff --git a/Engine/Source/Editor/Persona/Private/SAnimationEditorViewport.cpp b/Engine/Source/Editor/Persona/Private/SAnimationEditorViewport.cpp
--- a/Engine/Source/Editor/Persona/Private/SAnimationEditorViewport.cpp
+++ b/Engine/Source/Editor/Persona/Private/SAnimationEditorViewport.cpp
@@ -341,6 +341,12 @@ TSharedPtr<SWidget> SAnimationEditorViewport::BuildViewportToolbar()
 
 			ViewportToolbarContext.AddObject(ContextObject);
 		}
+
+		// Give the asset editor a chance to extend the context.
+		if (TSharedPtr<SAnimationEditorViewportTabBody> Tab = TabBodyPtr.Pin())
+		{
+			Tab->GetAssetEditorToolkit()->InitToolMenuContext(ViewportToolbarContext);
+		}
 	}
 
 	return UToolMenus::Get()->GenerateWidget(ViewportToolbarName, ViewportToolbarContext);
diff --git a/Engine/Source/Editor/PhysicsAssetEditor/Private/PhysicsAssetEditor.cpp b/Engine/Source/Editor/PhysicsAssetEditor/Private/PhysicsAssetEditor.cpp
--- a/Engine/Source/Editor/PhysicsAssetEditor/Private/PhysicsAssetEditor.cpp
+++ b/Engine/Source/Editor/PhysicsAssetEditor/Private/PhysicsAssetEditor.cpp
@@ -879,119 +879,205 @@ void FPhysicsAssetEditor::ExtendMenu()
 
 void FPhysicsAssetEditor::ExtendViewportMenus()
 {
-	FToolMenuOwnerScoped OwnerScoped(this);
-
-	static const FName CharacterMenuName("Persona.AnimViewportCharacterMenu");
-	UToolMenu* ExtendableCharacterMenu = UToolMenus::Get()->ExtendMenu(CharacterMenuName);
-	ExtendableCharacterMenu->AddDynamicSection("PhysicsCharacterMenu", FNewToolMenuDelegate::CreateLambda([](UToolMenu* CharacterMenu)
+	auto ExtendMenuWithPhysicsRenderingSection = [MenuOwner = this](FName InMenuName) -> void
 	{
-		TSharedPtr<FPhysicsAssetEditor> PhysicsAssetEditor = PhysicsAssetEditor::GetPhysicsAssetEditorFromToolContext(CharacterMenu->Context);
-		if (PhysicsAssetEditor)
-		{
-			FToolMenuSection& Section = CharacterMenu->AddSection("PhysicsAssetShowCommands", LOCTEXT("PhysicsShowCommands", "Physics Rendering"), FToolMenuInsert("AnimViewportSceneElements", EToolMenuInsertType::Before));
+		FToolMenuOwnerScoped OwnerScoped(MenuOwner);
 
-			Section.AddSubMenu(TEXT("MassPropertiesSubMenu"), LOCTEXT("MassPropertiesSubMenu", "Mass Properties"), FText::GetEmpty(),
-				FNewToolMenuDelegate::CreateLambda([WeakPhysicsAssetEditor = PhysicsAssetEditor.ToWeakPtr()](UToolMenu* InSubMenu)
+		UToolMenu* ExtendableCharacterMenu = UToolMenus::Get()->ExtendMenu(InMenuName);
+		ExtendableCharacterMenu->AddDynamicSection(
+			"PhysicsCharacterMenu",
+			FNewToolMenuDelegate::CreateLambda(
+				[](UToolMenu* CharacterMenu)
+				{
+					TSharedPtr<FPhysicsAssetEditor> PhysicsAssetEditor =
+						PhysicsAssetEditor::GetPhysicsAssetEditorFromToolContext(CharacterMenu->Context);
+					if (PhysicsAssetEditor)
 					{
-						const FPhysicsAssetEditorCommands& Commands = FPhysicsAssetEditorCommands::Get();
-
-						{
-							FToolMenuSection& Section = InSubMenu->AddSection("PhysicsAssetEditorCenterOfMassRenderSettings", LOCTEXT("CenterOfMassRenderSettingsHeader", "Center of Mass Drawing"));
-							Section.AddMenuEntry(Commands.DrawBodyMass);
-							Section.AddMenuEntry(Commands.HideCenterOfMassForKinematicBodies);
-							Section.AddEntry(FToolMenuEntry::InitWidget(TEXT("CoMMarkerScale"), WeakPhysicsAssetEditor.Pin()->MakeCoMMarkerScaleWidget(), LOCTEXT("CoMMarkerScaleLabel", "Marker Scale")));
-						}
+						FToolMenuSection& Section = CharacterMenu->AddSection(
+							"PhysicsAssetShowCommands",
+							LOCTEXT("PhysicsShowCommands", "Physics Rendering"),
+							FToolMenuInsert("AnimViewportSceneElements", EToolMenuInsertType::Before)
+						);
+
+						Section.AddSubMenu(
+							TEXT("MassPropertiesSubMenu"),
+							LOCTEXT("MassPropertiesSubMenu", "Mass Properties"),
+							FText::GetEmpty(),
+							FNewToolMenuDelegate::CreateLambda(
+								[WeakPhysicsAssetEditor = PhysicsAssetEditor.ToWeakPtr()](UToolMenu* InSubMenu)
+								{
+									const FPhysicsAssetEditorCommands& Commands = FPhysicsAssetEditorCommands::Get();
 
-						{
-							FToolMenuSection& Section = InSubMenu->AddSection("PhysicsAssetEditorCenterOfMassRenderingMode", LOCTEXT("CenterOfMassRenderingModeHeader", "Center of Mass Drawing (Edit)"));
-							Section.AddMenuEntry(Commands.CenterOfMassRenderingMode_All);
-							Section.AddMenuEntry(Commands.CenterOfMassRenderingMode_Selected);
-							Section.AddMenuEntry(Commands.CenterOfMassRenderingMode_None);
-						}
+									{
+										FToolMenuSection& Section = InSubMenu->AddSection(
+											"PhysicsAssetEditorCenterOfMassRenderSettings",
+											LOCTEXT("CenterOfMassRenderSettingsHeader", "Center of Mass Drawing")
+										);
+										Section.AddMenuEntry(Commands.DrawBodyMass);
+										Section.AddMenuEntry(Commands.HideCenterOfMassForKinematicBodies);
+										Section.AddEntry(FToolMenuEntry::InitWidget(
+											TEXT("CoMMarkerScale"),
+											WeakPhysicsAssetEditor.Pin()->MakeCoMMarkerScaleWidget(),
+											LOCTEXT("CoMMarkerScaleLabel", "Marker Scale")
+										));
+									}
 
-						{
-							FToolMenuSection& Section = InSubMenu->AddSection("PhysicsAssetEditorCenterOfMassRenderingModeSim", LOCTEXT("CenterOfMassRenderingModeSimHeader", "Center of Mass Drawing (Simulation)"));
-							Section.AddMenuEntry(Commands.CenterOfMassRenderingMode_Simulation_All);
-							Section.AddMenuEntry(Commands.CenterOfMassRenderingMode_Simulation_Selected);
-							Section.AddMenuEntry(Commands.CenterOfMassRenderingMode_Simulation_None);
-						}
-					}));
+									{
+										FToolMenuSection& Section = InSubMenu->AddSection(
+											"PhysicsAssetEditorCenterOfMassRenderingMode",
+											LOCTEXT("CenterOfMassRenderingModeHeader", "Center of Mass Drawing (Edit)")
+										);
+										Section.AddMenuEntry(Commands.CenterOfMassRenderingMode_All);
+										Section.AddMenuEntry(Commands.CenterOfMassRenderingMode_Selected);
+										Section.AddMenuEntry(Commands.CenterOfMassRenderingMode_None);
+									}
 
-			Section.AddSubMenu(TEXT("MeshRenderModeSubMenu"), LOCTEXT("MeshRenderModeSubMenu", "Mesh"), FText::GetEmpty(),
-				FNewToolMenuDelegate::CreateLambda([](UToolMenu* InSubMenu)
-			{
-				const FPhysicsAssetEditorCommands& Commands = FPhysicsAssetEditorCommands::Get();
-				{
-					FToolMenuSection& Section = InSubMenu->AddSection("PhysicsAssetEditorRenderingMode", LOCTEXT("MeshRenderModeHeader", "Mesh Drawing (Edit)"));
-					Section.AddMenuEntry(Commands.MeshRenderingMode_Solid);
-					Section.AddMenuEntry(Commands.MeshRenderingMode_Wireframe);
-					Section.AddMenuEntry(Commands.MeshRenderingMode_None);
-				}
+									{
+										FToolMenuSection& Section = InSubMenu->AddSection(
+											"PhysicsAssetEditorCenterOfMassRenderingModeSim",
+											LOCTEXT("CenterOfMassRenderingModeSimHeader", "Center of Mass Drawing (Simulation)")
+										);
+										Section.AddMenuEntry(Commands.CenterOfMassRenderingMode_Simulation_All);
+										Section.AddMenuEntry(Commands.CenterOfMassRenderingMode_Simulation_Selected);
+										Section.AddMenuEntry(Commands.CenterOfMassRenderingMode_Simulation_None);
+									}
+								}
+							)
+						);
+
+						Section.AddSubMenu(
+							TEXT("MeshRenderModeSubMenu"),
+							LOCTEXT("MeshRenderModeSubMenu", "Mesh"),
+							FText::GetEmpty(),
+							FNewToolMenuDelegate::CreateLambda(
+								[](UToolMenu* InSubMenu)
+								{
+									const FPhysicsAssetEditorCommands& Commands = FPhysicsAssetEditorCommands::Get();
+									{
+										FToolMenuSection& Section = InSubMenu->AddSection(
+											"PhysicsAssetEditorRenderingMode",
+											LOCTEXT("MeshRenderModeHeader", "Mesh Drawing (Edit)")
+										);
+										Section.AddMenuEntry(Commands.MeshRenderingMode_Solid);
+										Section.AddMenuEntry(Commands.MeshRenderingMode_Wireframe);
+										Section.AddMenuEntry(Commands.MeshRenderingMode_None);
+									}
 
-				{
-					FToolMenuSection& Section = InSubMenu->AddSection("PhysicsAssetEditorRenderingModeSim", LOCTEXT("MeshRenderModeSimHeader", "Mesh Drawing (Simulation)"));
-					Section.AddMenuEntry(Commands.MeshRenderingMode_Simulation_Solid);
-					Section.AddMenuEntry(Commands.MeshRenderingMode_Simulation_Wireframe);
-					Section.AddMenuEntry(Commands.MeshRenderingMode_Simulation_None);
-				}
-			}));
+									{
+										FToolMenuSection& Section = InSubMenu->AddSection(
+											"PhysicsAssetEditorRenderingModeSim",
+											LOCTEXT("MeshRenderModeSimHeader", "Mesh Drawing (Simulation)")
+										);
+										Section.AddMenuEntry(Commands.MeshRenderingMode_Simulation_Solid);
+										Section.AddMenuEntry(Commands.MeshRenderingMode_Simulation_Wireframe);
+										Section.AddMenuEntry(Commands.MeshRenderingMode_Simulation_None);
+									}
+								}
+							)
+						);
+
+						Section.AddSubMenu(
+							TEXT("CollisionRenderModeSubMenu"),
+							LOCTEXT("CollisionRenderModeSubMenu", "Bodies"),
+							FText::GetEmpty(),
+							FNewToolMenuDelegate::CreateLambda(
+								[WeakPhysicsAssetEditor = PhysicsAssetEditor.ToWeakPtr()](UToolMenu* InSubMenu)
+								{
+									const FPhysicsAssetEditorCommands& Commands = FPhysicsAssetEditorCommands::Get();
+									{
+										FToolMenuSection& Section = InSubMenu->AddSection(
+											"PhysicsAssetEditorCollisionRenderSettings",
+											LOCTEXT("CollisionRenderSettingsHeader", "Body Drawing")
+										);
+										Section.AddMenuEntry(Commands.RenderOnlySelectedSolid);
+										Section.AddMenuEntry(Commands.HideSimulatedBodies);
+										Section.AddMenuEntry(Commands.HideKinematicBodies);
+										Section.AddEntry(FToolMenuEntry::InitWidget(
+											TEXT("CollisionOpacity"),
+											WeakPhysicsAssetEditor.Pin()->MakeCollisionOpacityWidget(),
+											LOCTEXT("CollisionOpacityLabel", "Collision Opacity")
+										));
+									}
 
-			Section.AddSubMenu(TEXT("CollisionRenderModeSubMenu"), LOCTEXT("CollisionRenderModeSubMenu", "Bodies"), FText::GetEmpty(),
-				FNewToolMenuDelegate::CreateLambda([WeakPhysicsAssetEditor = PhysicsAssetEditor.ToWeakPtr()](UToolMenu* InSubMenu)
-			{
-				const FPhysicsAssetEditorCommands& Commands = FPhysicsAssetEditorCommands::Get();
-				{
-					FToolMenuSection& Section = InSubMenu->AddSection("PhysicsAssetEditorCollisionRenderSettings", LOCTEXT("CollisionRenderSettingsHeader", "Body Drawing"));
-					Section.AddMenuEntry(Commands.RenderOnlySelectedSolid);
-					Section.AddMenuEntry(Commands.HideSimulatedBodies);
-					Section.AddMenuEntry(Commands.HideKinematicBodies);
-					Section.AddEntry(FToolMenuEntry::InitWidget(TEXT("CollisionOpacity"), WeakPhysicsAssetEditor.Pin()->MakeCollisionOpacityWidget(), LOCTEXT("CollisionOpacityLabel", "Collision Opacity")));
-				}
+									{
+										FToolMenuSection& Section = InSubMenu->AddSection(
+											"PhysicsAssetEditorCollisionMode",
+											LOCTEXT("CollisionRenderModeHeader", "Body Drawing (Edit)")
+										);
+										Section.AddMenuEntry(Commands.CollisionRenderingMode_Solid);
+										Section.AddMenuEntry(Commands.CollisionRenderingMode_Wireframe);
+										Section.AddMenuEntry(Commands.CollisionRenderingMode_SolidWireframe);
+										Section.AddMenuEntry(Commands.CollisionRenderingMode_None);
+									}
 
-				{
-					FToolMenuSection& Section = InSubMenu->AddSection("PhysicsAssetEditorCollisionMode", LOCTEXT("CollisionRenderModeHeader", "Body Drawing (Edit)"));
-					Section.AddMenuEntry(Commands.CollisionRenderingMode_Solid);
-					Section.AddMenuEntry(Commands.CollisionRenderingMode_Wireframe);
-					Section.AddMenuEntry(Commands.CollisionRenderingMode_SolidWireframe);
-					Section.AddMenuEntry(Commands.CollisionRenderingMode_None);
-				}
+									{
+										FToolMenuSection& Section = InSubMenu->AddSection(
+											"PhysicsAssetEditorCollisionModeSim",
+											LOCTEXT("CollisionRenderModeSimHeader", "Body Drawing (Simulation)")
+										);
+										Section.AddMenuEntry(Commands.CollisionRenderingMode_Simulation_Solid);
+										Section.AddMenuEntry(Commands.CollisionRenderingMode_Simulation_Wireframe);
+										Section.AddMenuEntry(Commands.CollisionRenderingMode_Simulation_SolidWireframe);
+										Section.AddMenuEntry(Commands.CollisionRenderingMode_Simulation_None);
+									}
+								}
+							)
+						);
+
+						Section.AddSubMenu(
+							TEXT("ConstraintConstraintModeSubMenu"),
+							LOCTEXT("ConstraintConstraintModeSubMenu", "Constraints"),
+							FText::GetEmpty(),
+							FNewToolMenuDelegate::CreateLambda(
+								[WeakPhysicsAssetEditor = PhysicsAssetEditor.ToWeakPtr()](UToolMenu* InSubMenu)
+								{
+									const FPhysicsAssetEditorCommands& Commands = FPhysicsAssetEditorCommands::Get();
+									{
+										FToolMenuSection& Section = InSubMenu->AddSection(
+											"PhysicsAssetEditorConstraints", LOCTEXT("ConstraintHeader", "Constraints")
+										);
+										Section.AddMenuEntry(Commands.DrawConstraintsAsPoints);
+										Section.AddMenuEntry(Commands.DrawViolatedLimits);
+										Section.AddMenuEntry(Commands.RenderOnlySelectedConstraints);
+										Section.AddEntry(FToolMenuEntry::InitWidget(
+											TEXT("ConstraintScale"),
+											WeakPhysicsAssetEditor.Pin()->MakeConstraintScaleWidget(),
+											LOCTEXT("ConstraintScaleLabel", "Constraint Scale")
+										));
+									}
+									{
+										FToolMenuSection& Section = InSubMenu->AddSection(
+											"PhysicsAssetEditorConstraintMode",
+											LOCTEXT("ConstraintRenderModeHeader", "Constraint Drawing (Edit)")
+										);
+										Section.AddMenuEntry(Commands.ConstraintRenderingMode_None);
+										Section.AddMenuEntry(Commands.ConstraintRenderingMode_AllPositions);
+										Section.AddMenuEntry(Commands.ConstraintRenderingMode_AllLimits);
+									}
 
-				{
-					FToolMenuSection& Section = InSubMenu->AddSection("PhysicsAssetEditorCollisionModeSim", LOCTEXT("CollisionRenderModeSimHeader", "Body Drawing (Simulation)"));
-					Section.AddMenuEntry(Commands.CollisionRenderingMode_Simulation_Solid);
-					Section.AddMenuEntry(Commands.CollisionRenderingMode_Simulation_Wireframe);
-					Section.AddMenuEntry(Commands.CollisionRenderingMode_Simulation_SolidWireframe);
-					Section.AddMenuEntry(Commands.CollisionRenderingMode_Simulation_None);
+									{
+										FToolMenuSection& Section = InSubMenu->AddSection(
+											"PhysicsAssetEditorConstraintModeSim",
+											LOCTEXT("ConstraintRenderModeSimHeader", "Constraint Drawing (Simulation)")
+										);
+										Section.AddMenuEntry(Commands.ConstraintRenderingMode_Simulation_None);
+										Section.AddMenuEntry(Commands.ConstraintRenderingMode_Simulation_AllPositions);
+										Section.AddMenuEntry(Commands.ConstraintRenderingMode_Simulation_AllLimits);
+									}
+								}
+							)
+						);
+					}
 				}
-			}));
+			)
+		);
+	};
 
-			Section.AddSubMenu(TEXT("ConstraintConstraintModeSubMenu"), LOCTEXT("ConstraintConstraintModeSubMenu", "Constraints"), FText::GetEmpty(),
-				FNewToolMenuDelegate::CreateLambda([WeakPhysicsAssetEditor = PhysicsAssetEditor.ToWeakPtr()](UToolMenu* InSubMenu)
-			{
-				const FPhysicsAssetEditorCommands& Commands = FPhysicsAssetEditorCommands::Get();
-				{
-					FToolMenuSection& Section = InSubMenu->AddSection("PhysicsAssetEditorConstraints", LOCTEXT("ConstraintHeader", "Constraints"));
-					Section.AddMenuEntry(Commands.DrawConstraintsAsPoints);
-					Section.AddMenuEntry(Commands.DrawViolatedLimits);
-					Section.AddMenuEntry(Commands.RenderOnlySelectedConstraints);
-					Section.AddEntry(FToolMenuEntry::InitWidget(TEXT("ConstraintScale"), WeakPhysicsAssetEditor.Pin()->MakeConstraintScaleWidget(), LOCTEXT("ConstraintScaleLabel", "Constraint Scale")));
-				}
-				{
-					FToolMenuSection& Section = InSubMenu->AddSection("PhysicsAssetEditorConstraintMode", LOCTEXT("ConstraintRenderModeHeader", "Constraint Drawing (Edit)"));
-					Section.AddMenuEntry(Commands.ConstraintRenderingMode_None);
-					Section.AddMenuEntry(Commands.ConstraintRenderingMode_AllPositions);
-					Section.AddMenuEntry(Commands.ConstraintRenderingMode_AllLimits);
-				}
+	// Extend the old viewport toolbar "Character" menu.
+	ExtendMenuWithPhysicsRenderingSection("Persona.AnimViewportCharacterMenu");
+	// Extend the new viewport toolbar "Show" menu.
+	ExtendMenuWithPhysicsRenderingSection("AnimationEditor.ViewportToolbar.Show");
 
-				{
-					FToolMenuSection& Section = InSubMenu->AddSection("PhysicsAssetEditorConstraintModeSim", LOCTEXT("ConstraintRenderModeSimHeader", "Constraint Drawing (Simulation)"));
-					Section.AddMenuEntry(Commands.ConstraintRenderingMode_Simulation_None);
-					Section.AddMenuEntry(Commands.ConstraintRenderingMode_Simulation_AllPositions);
-					Section.AddMenuEntry(Commands.ConstraintRenderingMode_Simulation_AllLimits);
-				}
-			}));
-		}
-	}));
 	static const FName PhysicsMenuName("Persona.AnimViewportPhysicsMenu");
 	UToolMenu* ExtendablePhysicsMenu = UToolMenus::Get()->ExtendMenu(PhysicsMenuName);
 	ExtendablePhysicsMenu->AddDynamicSection("AnimViewportPhysicsMenu", FNewToolMenuDelegate::CreateLambda([](UToolMenu* PhysicsMenu)
