[UE] Fix for PhAT delete body selected in skeleton tree

#rnx

diff --git a/Engine/Source/Editor/PhysicsAssetEditor/Private/PhysicsAssetEditorSharedData.cpp b/Engine/Source/Editor/PhysicsAssetEditor/Private/PhysicsAssetEditorSharedData.cpp
--- a/Engine/Source/Editor/PhysicsAssetEditor/Private/PhysicsAssetEditorSharedData.cpp
+++ b/Engine/Source/Editor/PhysicsAssetEditor/Private/PhysicsAssetEditorSharedData.cpp
@@ -151,6 +151,11 @@ void ForEachUniquePair(CollectionType& Collection, const FunctionObjectType& Fun
 	}
 }
 
+bool SelectionContainsIndex(const TArray<FPhysicsAssetEditorSharedData::FSelection>& InSelectedElements, const int32 InIndex)
+{
+	return InSelectedElements.ContainsByPredicate([InIndex](const FPhysicsAssetEditorSharedData::FSelection& InSelection) {return InSelection.Index == InIndex; }); // Predicate returns true if the index is already in the selection.
+}
+
 // class FScopedBulkSelection //
 
 FScopedBulkSelection::FScopedBulkSelection(TSharedPtr<FPhysicsAssetEditorSharedData> InSharedData)
@@ -942,7 +947,8 @@ void FPhysicsAssetEditorSharedData::SetSelectedBodiesFromDirectlInput(const TArr
 	// Create a selection for each body that references the first primitive in that body. (ordering of primitives is determined by geometry type then by order if instantiation).
 	TArray<FSelection> SelectionCollection = CreateBodyPrimitivesSelection(PhysicsAsset, BodiesIndices, [this](const TArray<FSelection>& CurrentSelection, const int32 BodyIndex, const FKShapeElem& Primitive) 
 		{ 
-			return !IsBodySelected(BodyIndex); 
+			// Select the first primitive in the supplied body. 
+			return !SelectionContainsIndex(CurrentSelection, BodyIndex); // Predicate returns true if the supplied body index is not already in the current selection.
 		});
 
 	for (FSelection& Selection : SelectionCollection)
@@ -2906,23 +2912,35 @@ void FPhysicsAssetEditorSharedData::DeleteCurrentSelection()
 
 void FPhysicsAssetEditorSharedData::DeleteCurrentBody()
 {
-	// Delete any directly selected bodies and all their primitives
+	// Delete any directly selected bodies and all their primitives.
 
-	TArray<FSelection> SelectionsToRemove;
+	TArray<FSelection> DirectSelectedBodies;
 
 	for (const FSelection& Selection : SelectedBodies)
 	{
 		if (Selection.bIsDirectBodySelection)
 		{
-			DeleteBody(Selection.Index, false);
-			SelectionsToRemove.Add(Selection);
+			DirectSelectedBodies.Add(Selection);
 		}
 	}
 
-	SetSelectedBodies(SelectionsToRemove, false);
+	if (!DirectSelectedBodies.IsEmpty())
+	{
+		// Remove target body indexes from the selection.
+		SetSelectedBodies(DirectSelectedBodies, false);
 
-	RefreshPhysicsAssetChange(PhysicsAsset);
-	BroadcastHierarchyChanged();
+		// Sort by body index - highest first - as body indexes greater than the deleted index in the physics asset will be modified by each deletion.
+		Algo::SortBy(DirectSelectedBodies, &FSelection::Index, TGreater<>());
+
+		// Delete target bodies.
+		for (const FSelection& Selection : DirectSelectedBodies)
+		{
+			DeleteBody(Selection.Index, false);
+		}
+
+		RefreshPhysicsAssetChange(PhysicsAsset);
+		BroadcastHierarchyChanged();
+	}
 }
 
 void FPhysicsAssetEditorSharedData::DeleteCurrentPrim()
