[UE] PhAT : Fix for bounding bodies viewport text generation with zero bodies

diff --git a/Engine/Source/Editor/PhysicsAssetEditor/Private/PhysicsAssetEditorMode.cpp b/Engine/Source/Editor/PhysicsAssetEditor/Private/PhysicsAssetEditorMode.cpp
--- a/Engine/Source/Editor/PhysicsAssetEditor/Private/PhysicsAssetEditorMode.cpp
+++ b/Engine/Source/Editor/PhysicsAssetEditor/Private/PhysicsAssetEditorMode.cpp
@@ -192,13 +192,24 @@ FPhysicsAssetEditorMode::FPhysicsAssetEditorMode(TSharedRef<FWorkflowCentricAppl
 		if(InViewportCorner == EViewportCorner::TopLeft)
 		{
 			TSharedPtr<FPhysicsAssetEditorSharedData> SharedData = PhysicsAssetEditorPtr.Pin()->GetSharedData();
+			
+			FString BodiesConsideredForBoundsText;
+
+			const int32 TotalBodyCount = SharedData->PhysicsAsset->SkeletalBodySetups.Num();
+			if (TotalBodyCount > 0)
+			{
+				BodiesConsideredForBoundsText += " (";
+				BodiesConsideredForBoundsText += FString::FromInt(SharedData->PhysicsAsset->BoundsBodies.Num());
+				BodiesConsideredForBoundsText += " Considered For Bounds, ";
+				BodiesConsideredForBoundsText += FString::FromInt((SharedData->PhysicsAsset->BoundsBodies.Num() * 100) / TotalBodyCount);
+				BodiesConsideredForBoundsText += "%)";
+			}
 
 			// Write physics asset summary at the top of the view port.
 			return FText::Format(
-				NSLOCTEXT("UnrealEd", "BodiesConstraints_F", "{0} Bodies ({1} Considered for bounds, {2}%)\n{3}\n{4} Constraints{5}\n{6} Collision Interactions"),
-				FText::AsNumber(SharedData->PhysicsAsset->SkeletalBodySetups.Num()),
-				FText::AsNumber(SharedData->PhysicsAsset->BoundsBodies.Num()),
-				FText::AsNumber(static_cast<float>(SharedData->PhysicsAsset->BoundsBodies.Num()) / static_cast<float>(SharedData->PhysicsAsset->SkeletalBodySetups.Num()) * 100.0f),
+				NSLOCTEXT("UnrealEd", "BodiesConstraints_F", "{0} Bodies{1}\n{2}\n{3} Constraints{4}\n{5} Collision Interactions"),
+				FText::AsNumber(SharedData->PhysicsAsset->SkeletalBodySetups.Num()),				
+				FText::FromString(BodiesConsideredForBoundsText),
 				BuildPhysicsAssetShapeTypeCountText(SharedData->PhysicsAsset),
 				FText::AsNumber(SharedData->PhysicsAsset->ConstraintSetup.Num()),
 				BuildCrossConstraintCountText(SharedData->PhysicsAsset),
