[UE] PhAT : Prevent deselecting bodies when holding shift key for multi select and clicking on empty space.

diff --git a/Engine/Source/Editor/PhysicsAssetEditor/Private/PhysicsAssetEditor.cpp b/Engine/Source/Editor/PhysicsAssetEditor/Private/PhysicsAssetEditor.cpp
--- a/Engine/Source/Editor/PhysicsAssetEditor/Private/PhysicsAssetEditor.cpp
+++ b/Engine/Source/Editor/PhysicsAssetEditor/Private/PhysicsAssetEditor.cpp
@@ -4421,9 +4421,12 @@ void FPhysicsAssetEditor::HandleSelectionChanged(const TArrayView<TSharedPtr<ISk
 			// Block selection broadcast until we have selected all, as this can be an expensive operation
 			FScopedBulkSelection BulkSelection(SharedData);
 
-			// clear selection
-			SharedData->ClearSelectedBody();
-			SharedData->ClearSelectedConstraints();
+			if (!SharedData->IsGroupSelectionActive())
+			{
+				// clear selection
+				SharedData->ClearSelectedBody();
+				SharedData->ClearSelectedConstraints();
+			}
 
 			bool bBoneSelected = false;
 			TArray<FPhysicsAssetEditorSharedData::FSelection> SelectedPrimitives;
diff --git a/Engine/Source/Editor/PhysicsAssetEditor/Private/PhysicsAssetEditorEditMode.cpp b/Engine/Source/Editor/PhysicsAssetEditor/Private/PhysicsAssetEditorEditMode.cpp
--- a/Engine/Source/Editor/PhysicsAssetEditor/Private/PhysicsAssetEditorEditMode.cpp
+++ b/Engine/Source/Editor/PhysicsAssetEditor/Private/PhysicsAssetEditorEditMode.cpp
@@ -622,6 +622,8 @@ void FPhysicsAssetEditorEditMode::Tick(FEditorViewportClient* ViewportClient, fl
 		}	
 	}
 
+	SharedData->SetGroupSelectionActive(ViewportClient && (ViewportClient->IsCtrlPressed() || ViewportClient->IsShiftPressed()));
+
 	UpdateCoM();
 }
 
@@ -863,26 +865,26 @@ bool FPhysicsAssetEditorEditMode::HandleClick(FEditorViewportClient* InViewportC
 			{
 				HPhysicsAssetEditorEdBoneProxy* BoneProxy = (HPhysicsAssetEditorEdBoneProxy*)HitProxy;
 
-				SharedData->HitBone(BoneProxy->BodyIndex, BoneProxy->PrimType, BoneProxy->PrimIndex, InViewportClient->IsCtrlPressed() || InViewportClient->IsShiftPressed());
+				SharedData->HitBone(BoneProxy->BodyIndex, BoneProxy->PrimType, BoneProxy->PrimIndex, IsGroupSelectionActive());
 				return true;
 			}
 			else if (HitProxy && HitProxy->IsA(HPhysicsAssetEditorEdConstraintProxy::StaticGetType()))
 			{
 				HPhysicsAssetEditorEdConstraintProxy* ConstraintProxy = (HPhysicsAssetEditorEdConstraintProxy*)HitProxy;
 
-				SharedData->HitConstraint(ConstraintProxy->ConstraintIndex, InViewportClient->IsCtrlPressed() || InViewportClient->IsShiftPressed());
+				SharedData->HitConstraint(ConstraintProxy->ConstraintIndex, IsGroupSelectionActive());
 				return true;
 			}
 			else if (HitProxy && HitProxy->IsA(HPhysicsAssetEditorEdCoMProxy::StaticGetType()))
 			{
 				HPhysicsAssetEditorEdCoMProxy* CoMProxy = (HPhysicsAssetEditorEdCoMProxy*)HitProxy;
 
-				SharedData->HitCoM(CoMProxy->BodyIndex, InViewportClient->IsCtrlPressed() || InViewportClient->IsShiftPressed());
+				SharedData->HitCoM(CoMProxy->BodyIndex, IsGroupSelectionActive());
 				return true;
 			}
 			else
 			{
-				HitNothing(InViewportClient);
+				HitNothing(InViewportClient, IsGroupSelectionActive());
 			}
 		}
 		else if (Click.GetKey() == EKeys::RightMouseButton)
@@ -1214,9 +1216,14 @@ void FPhysicsAssetEditorEditMode::ModifyPrimitiveSize(int32 BodyIndex, EAggColli
 	}
 }
 
-void FPhysicsAssetEditorEditMode::HitNothing(FEditorViewportClient* InViewportClient)
+bool FPhysicsAssetEditorEditMode::IsGroupSelectionActive() const
+{
+	return SharedData && SharedData->IsGroupSelectionActive();
+}
+
+void FPhysicsAssetEditorEditMode::HitNothing(FEditorViewportClient* InViewportClient, const bool bGroupSelect)
 {
-	if (InViewportClient->IsCtrlPressed() == false)	//we only want to deselect if Ctrl is not used
+	if (!bGroupSelect)	//we only want to de-select if input was not part of a group select.
 	{
 		SharedData->ClearSelectedBody();
 		SharedData->ClearSelectedConstraints();
diff --git a/Engine/Source/Editor/PhysicsAssetEditor/Private/PhysicsAssetEditorEditMode.h b/Engine/Source/Editor/PhysicsAssetEditor/Private/PhysicsAssetEditorEditMode.h
--- a/Engine/Source/Editor/PhysicsAssetEditor/Private/PhysicsAssetEditorEditMode.h
+++ b/Engine/Source/Editor/PhysicsAssetEditor/Private/PhysicsAssetEditorEditMode.h
@@ -72,8 +72,11 @@ private:
 	/** Scales a collision body */
 	void ModifyPrimitiveSize(int32 BodyIndex, EAggCollisionShape::Type PrimType, int32 PrimIndex, FVector DeltaSize);
 
+	/** Returns true if input events should be considered part of a group selection operation. */
+	bool IsGroupSelectionActive() const;
+
 	/** Called when no scene proxy is hit, deselects everything */
-	void HitNothing(FEditorViewportClient* InViewportClient);
+	void HitNothing(FEditorViewportClient* InViewportClient, const bool bGroupSelect = false);
 
 	void CycleTransformMode();
 
diff --git a/Engine/Source/Editor/PhysicsAssetEditor/Private/PhysicsAssetEditorSharedData.cpp b/Engine/Source/Editor/PhysicsAssetEditor/Private/PhysicsAssetEditorSharedData.cpp
--- a/Engine/Source/Editor/PhysicsAssetEditor/Private/PhysicsAssetEditorSharedData.cpp
+++ b/Engine/Source/Editor/PhysicsAssetEditor/Private/PhysicsAssetEditorSharedData.cpp
@@ -1057,6 +1057,16 @@ void FPhysicsAssetEditorSharedData::SetSelectedBodies(const TArray<FSelection>&
 	}
 }
 
+void FPhysicsAssetEditorSharedData::SetGroupSelectionActive(const bool bIsActive)
+{
+	bIsGroupSelectionActive = bIsActive;
+}
+
+bool FPhysicsAssetEditorSharedData::IsGroupSelectionActive() const
+{
+	return bIsGroupSelectionActive;
+}
+
 bool FPhysicsAssetEditorSharedData::IsBodySelected(const FSelection& Body) const
 {
 	return SelectedBodies.Contains(Body);
diff --git a/Engine/Source/Editor/PhysicsAssetEditor/Private/PhysicsAssetEditorSharedData.h b/Engine/Source/Editor/PhysicsAssetEditor/Private/PhysicsAssetEditorSharedData.h
--- a/Engine/Source/Editor/PhysicsAssetEditor/Private/PhysicsAssetEditorSharedData.h
+++ b/Engine/Source/Editor/PhysicsAssetEditor/Private/PhysicsAssetEditorSharedData.h
@@ -187,6 +187,9 @@ public:
 	void DeleteCurrentConstraint();
 
 
+	void SetGroupSelectionActive(const bool bIsActive);
+	bool IsGroupSelectionActive() const;
+
 	/** Center of Mass editing */
 	void ClearSelectedCoMs();
 	void SetSelectedCoM(const FSelection& InSelectedCoM, const bool bSelected);
@@ -455,6 +458,7 @@ public:
 	/** Manipulation (rotate, translate, scale) */
 	bool bManipulating;
 
+	bool bIsGroupSelectionActive = false;
 	bool bShouldUpdatedSelectedCoMs = false;
 
 	/** when true, we dont broadcast every selection change - allows for bulk changes without so much overhead */
