[UE] Fix for static analysis warning - horde issue 726228

diff --git a/Engine/Source/Editor/PhysicsAssetEditor/Private/PhysicsAssetEditorEditMode.cpp b/Engine/Source/Editor/PhysicsAssetEditor/Private/PhysicsAssetEditorEditMode.cpp
--- a/Engine/Source/Editor/PhysicsAssetEditor/Private/PhysicsAssetEditorEditMode.cpp
+++ b/Engine/Source/Editor/PhysicsAssetEditor/Private/PhysicsAssetEditorEditMode.cpp
@@ -608,7 +608,9 @@ void FPhysicsAssetEditorEditMode::Tick(FEditorViewportClient* ViewportClient, fl
 	}
 
 	// Update the constraint view port manipulation flags from state of the modifier keys. These flags determine which parts of the constraint transform (Parent, child or both) should be modified when a view port widget is manipulated.
-	if (FPhysicsAssetRenderSettings* const RenderSettings = SharedData->GetRenderSettings())
+	FPhysicsAssetRenderSettings* const RenderSettings = SharedData->GetRenderSettings();
+	
+	if(ViewportClient && RenderSettings)
 	{
 		RenderSettings->ConstraintViewportManipulationFlags = EConstraintTransformComponentFlags::All;
 
