[UE] Fix for static analysis warning
- Issue 726818

diff --git a/Engine/Source/Editor/PhysicsAssetEditor/Private/PhysicsAssetEditorEditMode.cpp b/Engine/Source/Editor/PhysicsAssetEditor/Private/PhysicsAssetEditorEditMode.cpp
--- a/Engine/Source/Editor/PhysicsAssetEditor/Private/PhysicsAssetEditorEditMode.cpp
+++ b/Engine/Source/Editor/PhysicsAssetEditor/Private/PhysicsAssetEditorEditMode.cpp
@@ -576,55 +576,58 @@ bool FPhysicsAssetEditorEditMode::InputDelta(FEditorViewportClient* InViewportCl
 
 void FPhysicsAssetEditorEditMode::Tick(FEditorViewportClient* ViewportClient, float DeltaTime)
 {
-	if (SharedData->bRunningSimulation)
+	if (ViewportClient)
 	{
-		// check if PIE disabled the realtime viewport and quit sim if so
-		if (!ViewportClient->IsRealtime())
+		if (SharedData->bRunningSimulation)
 		{
-			SharedData->ToggleSimulation();
+			// check if PIE disabled the realtime viewport and quit sim if so
+			if (!ViewportClient->IsRealtime())
+			{
+				SharedData->ToggleSimulation();
 
-			ViewportClient->Invalidate();
-		}
+				ViewportClient->Invalidate();
+			}
 
-		UWorld* World = SharedData->PreviewScene.Pin()->GetWorld();
-		AWorldSettings* Setting = World->GetWorldSettings();
-		Setting->WorldGravityZ = SharedData->bNoGravitySimulation ? 0.0f : (SharedData->EditorOptions->bUseGravityOverride ? SharedData->EditorOptions->GravityOverrideZ : (UPhysicsSettings::Get()->DefaultGravityZ * SharedData->EditorOptions->GravScale));
-		Setting->bWorldGravitySet = true;
+			UWorld* World = SharedData->PreviewScene.Pin()->GetWorld();
+			AWorldSettings* Setting = World->GetWorldSettings();
+			Setting->WorldGravityZ = SharedData->bNoGravitySimulation ? 0.0f : (SharedData->EditorOptions->bUseGravityOverride ? SharedData->EditorOptions->GravityOverrideZ : (UPhysicsSettings::Get()->DefaultGravityZ * SharedData->EditorOptions->GravScale));
+			Setting->bWorldGravitySet = true;
 
-		// We back up the transforms array now
-		SharedData->EditorSkelComp->AnimationSpaceBases = SharedData->EditorSkelComp->GetComponentSpaceTransforms();
-		// When using the World solver, we must specify how much of the solver output gets blended into the animated mesh pose
-		// When using other solvers in PhAT, we don't want SetPhysicsBlendWeight function to re-enale the main solver physics
-		if (SharedData->PhysicsAsset->SolverType == EPhysicsAssetSolverType::World)
-		{
-			SharedData->EditorSkelComp->SetPhysicsBlendWeight(SharedData->EditorOptions->PhysicsBlend);
-		}
-		SharedData->EditorSkelComp->bUpdateJointsFromAnimation = SharedData->EditorOptions->bUpdateJointsFromAnimation;
-		SharedData->EditorSkelComp->PhysicsTransformUpdateMode = SharedData->EditorOptions->PhysicsUpdateMode;
+			// We back up the transforms array now
+			SharedData->EditorSkelComp->AnimationSpaceBases = SharedData->EditorSkelComp->GetComponentSpaceTransforms();
+			// When using the World solver, we must specify how much of the solver output gets blended into the animated mesh pose
+			// When using other solvers in PhAT, we don't want SetPhysicsBlendWeight function to re-enale the main solver physics
+			if (SharedData->PhysicsAsset->SolverType == EPhysicsAssetSolverType::World)
+			{
+				SharedData->EditorSkelComp->SetPhysicsBlendWeight(SharedData->EditorOptions->PhysicsBlend);
+			}
+			SharedData->EditorSkelComp->bUpdateJointsFromAnimation = SharedData->EditorOptions->bUpdateJointsFromAnimation;
+			SharedData->EditorSkelComp->PhysicsTransformUpdateMode = SharedData->EditorOptions->PhysicsUpdateMode;
 
-		static FPhysicalAnimationData EmptyProfile;
+			static FPhysicalAnimationData EmptyProfile;
 
-		SharedData->PhysicalAnimationComponent->ApplyPhysicalAnimationProfileBelow(NAME_None, SharedData->PhysicsAsset->CurrentPhysicalAnimationProfileName, /*Include Self=*/true, /*Clear Not Found=*/true);
-	}
+			SharedData->PhysicalAnimationComponent->ApplyPhysicalAnimationProfileBelow(NAME_None, SharedData->PhysicsAsset->CurrentPhysicalAnimationProfileName, /*Include Self=*/true, /*Clear Not Found=*/true);
+		}
 
-	// Update the constraint view port manipulation flags from state of the modifier keys. These flags determine which parts of the constraint transform (Parent, child or both) should be modified when a view port widget is manipulated.
-	FPhysicsAssetRenderSettings* const RenderSettings = SharedData->GetRenderSettings();
-	
-	if(ViewportClient && RenderSettings)
-	{
-		RenderSettings->ConstraintViewportManipulationFlags = EConstraintTransformComponentFlags::All;
+		// Update the constraint view port manipulation flags from state of the modifier keys. These flags determine which parts of the constraint transform (Parent, child or both) should be modified when a view port widget is manipulated.
+		FPhysicsAssetRenderSettings* const RenderSettings = SharedData->GetRenderSettings();
 
-		if (ViewportClient->IsPrioritizedInputChordPressed(InputChordName_EditConstraintChildTransform)) // Rotate or move the child transform only.
+		if (RenderSettings)
 		{
-			EnumRemoveFlags(RenderSettings->ConstraintViewportManipulationFlags, EConstraintTransformComponentFlags::AllParent); // Remove Parent Frame flags.
+			RenderSettings->ConstraintViewportManipulationFlags = EConstraintTransformComponentFlags::All;
+
+			if (ViewportClient->IsPrioritizedInputChordPressed(InputChordName_EditConstraintChildTransform)) // Rotate or move the child transform only.
+			{
+				EnumRemoveFlags(RenderSettings->ConstraintViewportManipulationFlags, EConstraintTransformComponentFlags::AllParent); // Remove Parent Frame flags.
+			}
+			else if (ViewportClient->IsAltPressed()) // Rotate or move the parent transform only.
+			{
+				EnumRemoveFlags(RenderSettings->ConstraintViewportManipulationFlags, EConstraintTransformComponentFlags::AllChild); // Remove Child Frame flags.
+			}
 		}
-		else if (ViewportClient->IsAltPressed()) // Rotate or move the parent transform only.
-		{
-			EnumRemoveFlags(RenderSettings->ConstraintViewportManipulationFlags, EConstraintTransformComponentFlags::AllChild); // Remove Child Frame flags.
-		}	
-	}
 
-	SharedData->SetGroupSelectionActive(ViewportClient && (ViewportClient->IsCtrlPressed() || ViewportClient->IsShiftPressed()));
+		SharedData->SetGroupSelectionActive(ViewportClient->IsCtrlPressed() || ViewportClient->IsShiftPressed());
+	}
 
 	UpdateCoM();
 }
