[UE] PhAT Prevent de-selecting constraints on right click
- Ensure tree selection is syncronised with viewport selection on click in viewport

diff --git a/Engine/Source/Editor/PhysicsAssetEditor/Private/PhysicsAssetEditorEditMode.cpp b/Engine/Source/Editor/PhysicsAssetEditor/Private/PhysicsAssetEditorEditMode.cpp
--- a/Engine/Source/Editor/PhysicsAssetEditor/Private/PhysicsAssetEditorEditMode.cpp
+++ b/Engine/Source/Editor/PhysicsAssetEditor/Private/PhysicsAssetEditorEditMode.cpp
@@ -862,6 +862,8 @@ FVector FPhysicsAssetEditorEditMode::GetWidgetLocation() const
 
 bool FPhysicsAssetEditorEditMode::HandleClick(FEditorViewportClient* InViewportClient, HHitProxy *HitProxy, const FViewportClick &Click)
 {
+	bool bHandled = false;
+
 	if (!SharedData->bRunningSimulation)
 	{
 		if (Click.GetKey() == EKeys::LeftMouseButton)
@@ -871,21 +873,21 @@ bool FPhysicsAssetEditorEditMode::HandleClick(FEditorViewportClient* InViewportC
 				HPhysicsAssetEditorEdBoneProxy* BoneProxy = (HPhysicsAssetEditorEdBoneProxy*)HitProxy;
 
 				SharedData->HitBone(BoneProxy->BodyIndex, BoneProxy->PrimType, BoneProxy->PrimIndex, IsGroupSelectionActive());
-				return true;
+				bHandled = true;
 			}
 			else if (HitProxy && HitProxy->IsA(HPhysicsAssetEditorEdConstraintProxy::StaticGetType()))
 			{
 				HPhysicsAssetEditorEdConstraintProxy* ConstraintProxy = (HPhysicsAssetEditorEdConstraintProxy*)HitProxy;
 
 				SharedData->HitConstraint(ConstraintProxy->ConstraintIndex, IsGroupSelectionActive());
-				return true;
+				bHandled = true;
 			}
 			else if (HitProxy && HitProxy->IsA(HPhysicsAssetEditorEdCoMProxy::StaticGetType()))
 			{
 				HPhysicsAssetEditorEdCoMProxy* CoMProxy = (HPhysicsAssetEditorEdCoMProxy*)HitProxy;
 
 				SharedData->HitCoM(CoMProxy->BodyIndex, IsGroupSelectionActive());
-				return true;
+				bHandled = true;
 			}
 			else
 			{
@@ -898,16 +900,10 @@ bool FPhysicsAssetEditorEditMode::HandleClick(FEditorViewportClient* InViewportC
 			{
 				HPhysicsAssetEditorEdBoneProxy* BoneProxy = (HPhysicsAssetEditorEdBoneProxy*)HitProxy;
 
-				// Select body under cursor if not already selected	(if ctrl is held down we only add, not remove)
+				// Select body under cursor if not already selected.
 				FPhysicsAssetEditorSharedData::FSelection Selection(BoneProxy->BodyIndex, BoneProxy->PrimType, BoneProxy->PrimIndex);
 				if (!SharedData->IsBodySelected(Selection))
 				{
-					if(!InViewportClient->IsCtrlPressed())
-					{
-						SharedData->ClearSelectedBody();
-						SharedData->ClearSelectedConstraints();
-					}
-
 					SharedData->SetSelectedBody(Selection, true);
 				}
 
@@ -917,21 +913,15 @@ bool FPhysicsAssetEditorEditMode::HandleClick(FEditorViewportClient* InViewportC
 					OpenBodyMenu(InViewportClient);
 				}
 
-				return true;
+				bHandled = true;
 			}
 			else if (HitProxy && HitProxy->IsA(HPhysicsAssetEditorEdConstraintProxy::StaticGetType()))
 			{
 				HPhysicsAssetEditorEdConstraintProxy* ConstraintProxy = (HPhysicsAssetEditorEdConstraintProxy*)HitProxy;
 
-				// Select constraint under cursor if not already selected (if ctrl is held down we only add, not remove)
+				// Select constraint under cursor if not already selected
 				if (!SharedData->IsConstraintSelected(ConstraintProxy->ConstraintIndex))
 				{
-					if(!InViewportClient->IsCtrlPressed())
-					{
-						SharedData->ClearSelectedBody();
-						SharedData->ClearSelectedConstraints();
-					}
-
 					SharedData->SetSelectedConstraint(ConstraintProxy->ConstraintIndex, true);
 				}
 
@@ -941,17 +931,19 @@ bool FPhysicsAssetEditorEditMode::HandleClick(FEditorViewportClient* InViewportC
 					OpenConstraintMenu(InViewportClient);
 				}
 
-				return true;
+				bHandled = true;
 			}
 			else
 			{
 				OpenSelectionMenu(InViewportClient);
-				return true;
+				bHandled = true;
 			}
 		}
 	}
 
-	return false;
+	SharedData->BroadcastSelectionChanged(); // Ensure selection in skeleton tree matches selection in viewport.
+
+	return bHandled;
 }
 
 bool FPhysicsAssetEditorEditMode::ReceivedFocus(FEditorViewportClient* ViewportClient, FViewport* Viewport)
