[UE] PhAT Fixes for mirroring
- Crash when mirror target bone has no body - Fix: Create a new physics body if one does not already exist for the mirror bone.
- Mirror target body has no constraints - Fix: Don't remove constraints when re-creating a physics body

diff --git a/Engine/Source/Developer/PhysicsUtilities/Private/PhysicsAssetUtils.cpp b/Engine/Source/Developer/PhysicsUtilities/Private/PhysicsAssetUtils.cpp
--- a/Engine/Source/Developer/PhysicsUtilities/Private/PhysicsAssetUtils.cpp
+++ b/Engine/Source/Developer/PhysicsUtilities/Private/PhysicsAssetUtils.cpp
@@ -1483,7 +1483,6 @@ int32 CreateNewBody(UPhysicsAsset* const PhysAsset, const FName InBodyName, cons
 void RecreateBody(UPhysicsAsset* const PhysAsset, const FName InBodyName, const FPhysAssetCreateParams& Params, const int32 BodyIndex)
 {
 	RemoveBodyFromCollisionTableInternal(PhysAsset, BodyIndex, false);
-	RemoveBodyConstraintsInternal(PhysAsset, BodyIndex);
 	PhysAsset->SkeletalBodySetups[BodyIndex] = TObjectPtr<USkeletalBodySetup>(); // Reset TObjectPtr to body setup that will be replaced.
 	CreateNewBodyAtIndexInternal(PhysAsset, InBodyName, Params, BodyIndex);
 }
diff --git a/Engine/Source/Editor/PhysicsAssetEditor/Private/PhysicsAssetEditorSharedData.cpp b/Engine/Source/Editor/PhysicsAssetEditor/Private/PhysicsAssetEditorSharedData.cpp
--- a/Engine/Source/Editor/PhysicsAssetEditor/Private/PhysicsAssetEditorSharedData.cpp
+++ b/Engine/Source/Editor/PhysicsAssetEditor/Private/PhysicsAssetEditorSharedData.cpp
@@ -638,9 +638,11 @@ void FPhysicsAssetEditorSharedData::Mirror()
 			{
 				UBodySetup * SrcBody = PhysicsAsset->SkeletalBodySetups[MirrorInfo.BodyIndex];
 				const FScopedTransaction Transaction(NSLOCTEXT("PhysicsAssetEditor", "MirrorBody", "MirrorBody"));
-				RecreateBody(MirrorBoneIndex, false); // Create a new body at the same index as the one being replaced. This ensures that all references to this body via index remain valid.
+				MakeOrRecreateBody(MirrorBoneIndex, false);
+
+				const int MirrorBodyIndex = PhysicsAsset->FindControllingBodyIndex(EditorSkelMesh, MirrorBoneIndex);
+				check(MirrorBodyIndex != INDEX_NONE);
 
-				const int32 MirrorBodyIndex = PhysicsAsset->FindControllingBodyIndex(EditorSkelMesh, MirrorBoneIndex);
 				UBodySetup* const DestBody = PhysicsAsset->SkeletalBodySetups[MirrorBodyIndex];
 				DestBody->Modify();
 				DestBody->CopyBodyPropertiesFrom(SrcBody);
@@ -2479,16 +2481,20 @@ void FPhysicsAssetEditorSharedData::RecreateBody(const FPhysAssetCreateParams& B
 	const FName BoneName = EditorSkelMesh->GetRefSkeleton().GetBoneName(BoneIndex);
 	const int32 BodyIndex = PhysicsAsset->FindBodyIndex(BoneName);
 
-	FPhysicsAssetUtils::RecreateBody(PhysicsAsset, BoneName, BodyData, BodyIndex); // Create a new physics body setup at the same index as the original body setup.
+	check(BodyIndex != INDEX_NONE);
+	if (BodyIndex != INDEX_NONE)
+	{
+		FPhysicsAssetUtils::RecreateBody(PhysicsAsset, BoneName, BodyData, BodyIndex); // Create a new physics body setup at the same index as the original body setup.
 
-	BroadcastHierarchyChanged();
+		BroadcastHierarchyChanged();
 
-	if (bAutoSelect)
-	{
-		SetSelectedBodyAnyPrimitive(BodyIndex, true);
-	}
+		if (bAutoSelect)
+		{
+			SetSelectedBodyAnyPrimitive(BodyIndex, true);
+		}
 
-	RefreshPhysicsAssetChange(PhysicsAsset);
+		RefreshPhysicsAssetChange(PhysicsAsset);
+	}
 }
 
 void FPhysicsAssetEditorSharedData::MakeNewBody(int32 NewBoneIndex, bool bAutoSelect)
@@ -2640,6 +2646,23 @@ void FPhysicsAssetEditorSharedData::MakeNewBody(const FPhysAssetCreateParams& Ne
 	RefreshPhysicsAssetChange(PhysicsAsset);
 }
 
+void FPhysicsAssetEditorSharedData::MakeOrRecreateBody(const int32 NewBoneIndex, const bool bAutoSelect)
+{
+	if (const USkeletalMesh* const EditorSkelMesh = PhysicsAsset->GetPreviewMesh())
+	{
+		const FName NewBoneName = EditorSkelMesh->GetRefSkeleton().GetBoneName(NewBoneIndex);
+
+		if (PhysicsAsset->FindBodyIndex(NewBoneName) != INDEX_NONE)
+		{
+			RecreateBody(NewBoneIndex, bAutoSelect); // Create a new body at the same index as the one being replaced. This ensures that all references to this body via index remain valid.
+		}
+		else
+		{
+			MakeNewBody(NewBoneIndex, bAutoSelect); // Create a new body.
+		}
+	}
+}
+
 FString FPhysicsAssetEditorSharedData::MakeUniqueNewConstraintName()
 {
 	// Make a new unique name for this constraint
@@ -3245,7 +3268,7 @@ void FPhysicsAssetEditorSharedData::EnableSimulation(bool bEnableSimulation)
 		// Stop any animation and clear node when stopping simulation.
 		PhysicalAnimationComponent->SetSkeletalMeshComponent(nullptr);
 
-		// Undo ends up recreating the anim script instance, so we need to remove it here (otherwise the AnimNode_RigidBody similation starts when we undo)
+		// Undo ends up recreating the anim script instance, so we need to remove it here (otherwise the AnimNode_RigidBody simulation starts when we undo)
 		EditorSkelComp->ClearAnimScriptInstance();
 
 		EditorSkelComp->SetPhysicsBlendWeight(0.f);
diff --git a/Engine/Source/Editor/PhysicsAssetEditor/Private/PhysicsAssetEditorSharedData.h b/Engine/Source/Editor/PhysicsAssetEditor/Private/PhysicsAssetEditorSharedData.h
--- a/Engine/Source/Editor/PhysicsAssetEditor/Private/PhysicsAssetEditorSharedData.h
+++ b/Engine/Source/Editor/PhysicsAssetEditor/Private/PhysicsAssetEditorSharedData.h
@@ -230,6 +230,7 @@ public:
 	void MakeNewBody(const FPhysAssetCreateParams& NewBodyData, const int32 NewBoneIndex, const bool bAutoSelect = true);
 	void RecreateBody(const int32 BoneIndex, const bool bAutoSelect = true); // Create a new physics body to replace an existing one.
 	void RecreateBody(const FPhysAssetCreateParams& BodyData, const int32 BoneIndex, const bool bAutoSelect = true); // Create a new physics body to replace an existing one.
+	void MakeOrRecreateBody(const int32 NewBoneIndex, const bool bAutoSelect = true); // Create a new physics body, replacing an existing one at the same index if it exists.
 	void MakeNewConstraints(int32 ParentBodyIndex, const TArray<int32>& ChildBodyIndices);
 	void MakeNewConstraint(int32 ParentBodyIndex, int32 ChildBodyIndex);
 	void CopySelectedBodiesAndConstraintsToClipboard(int32& OutNumCopiedBodies, int32& OutNumCopiedConstraints, int32& OutNumCopiedDisabledCollisionPairs);
