[UE] PhAT Make selection part of undo history
- Store selection in a new UObject derived class so that it can be supported by the undo system
- Option in PhAT edit menu to include selection in undo history

diff --git a/Engine/Source/Editor/PhysicsAssetEditor/Private/PhysicsAssetDetailsCustomization.cpp b/Engine/Source/Editor/PhysicsAssetEditor/Private/PhysicsAssetDetailsCustomization.cpp
--- a/Engine/Source/Editor/PhysicsAssetEditor/Private/PhysicsAssetDetailsCustomization.cpp
+++ b/Engine/Source/Editor/PhysicsAssetEditor/Private/PhysicsAssetDetailsCustomization.cpp
@@ -8,11 +8,12 @@
 #include "IDetailPropertyRow.h"
 #include "DetailWidgetRow.h"
 #include "Framework/MultiBox/MultiBoxBuilder.h"
-#include "PhysicsAssetEditorActions.h"
 #include "PhysicsAssetEditor.h"
+#include "PhysicsAssetEditorActions.h"
+#include "PhysicsAssetEditorSelection.h"
+#include "PhysicsAssetEditorSkeletalMeshComponent.h"
 #include "PhysicsEngine/PhysicsConstraintTemplate.h"
 #include "PhysicsEngine/SkeletalBodySetup.h"
-#include "PhysicsAssetEditorSkeletalMeshComponent.h"
 #include "Widgets/Text/SInlineEditableTextBlock.h"
 #include "EditorFontGlyphs.h"
 #include "Widgets/Layout/SUniformGridPanel.h"
@@ -688,10 +689,10 @@ void FPhysicsAssetDetailsCustomization::AddBodyToPhysicalAnimationProfile()
 	
 	TSharedPtr<FPhysicsAssetEditorSharedData> SharedData = PhysicsAssetEditorPtr.Pin()->GetSharedData();
 	UPhysicsAsset* PhysicsAsset = SharedData->PhysicsAsset;
-	for(int32 BodySetupIndex = 0; BodySetupIndex < SharedData->SelectedBodies.Num(); ++BodySetupIndex)
+		
+	for(const FPhysicsAssetEditorSharedData::FSelection& SelectedElement : SharedData->UniqueSelectionReferencingBodies())
 	{
-		USkeletalBodySetup* BodySetup = PhysicsAsset->SkeletalBodySetups[SharedData->SelectedBodies[BodySetupIndex].Index];
-		if(BodySetup)
+		if(USkeletalBodySetup* const BodySetup = PhysicsAsset->SkeletalBodySetups[SelectedElement.Index])
 		{
 			BodySetup->Modify();
 			FName ProfileName = BodySetup->GetCurrentPhysicalAnimationProfileName();
@@ -712,28 +713,32 @@ bool FPhysicsAssetDetailsCustomization::CanAddBodyToPhysicalAnimationProfile() c
 
 	auto PhysicalAnimationProfileExistsForAll = [WeakSharedData]()
 	{
-		TSharedPtr<FPhysicsAssetEditorSharedData> LocalSharedData = WeakSharedData.Pin();
-
-		for(int32 BodySetupIndex = 0; BodySetupIndex < LocalSharedData->SelectedBodies.Num(); ++BodySetupIndex)
+		if (TSharedPtr<FPhysicsAssetEditorSharedData> LocalSharedData = WeakSharedData.Pin())
 		{
-			USkeletalBodySetup* BodySetup = LocalSharedData->PhysicsAsset->SkeletalBodySetups[LocalSharedData->SelectedBodies[BodySetupIndex].Index];
-			if(BodySetup)
+			for (const FPhysicsAssetEditorSharedData::FSelection& SelectedElement : LocalSharedData->UniqueSelectionReferencingBodies())
 			{
-				if (!BodySetup->FindPhysicalAnimationProfile(BodySetup->GetCurrentPhysicalAnimationProfileName()))
+				const int32 BodySetupIndex = SelectedElement.Index;
+
+				check(LocalSharedData->PhysicsAsset->SkeletalBodySetups.IsValidIndex(BodySetupIndex));
+
+				if (const USkeletalBodySetup* const BodySetup = LocalSharedData->PhysicsAsset->SkeletalBodySetups[BodySetupIndex])
+				{
+					if (!BodySetup->FindPhysicalAnimationProfile(BodySetup->GetCurrentPhysicalAnimationProfileName()))
+					{
+						return false;
+					}
+				}
+				else
 				{
 					return false;
 				}
 			}
-			else
-			{
-				return false;
-			}
 		}
 
 		return true;
 	};
 
-	const bool bSelectedBodies = SharedData->SelectedBodies.Num() > 0;
+	const bool bSelectedBodies = !SharedData->UniqueSelectionReferencingBodies().IsEmpty();
 	return (PhysicsAssetEditorPtr.Pin()->IsNotSimulation() && bSelectedBodies && !PhysicalAnimationProfileExistsForAll() && PhysicsAsset->CurrentPhysicalAnimationProfileName != NAME_None);
 }
 
@@ -743,10 +748,10 @@ void FPhysicsAssetDetailsCustomization::RemoveBodyFromPhysicalAnimationProfile()
 
 	TSharedPtr<FPhysicsAssetEditorSharedData> SharedData = PhysicsAssetEditorPtr.Pin()->GetSharedData();
 	UPhysicsAsset* PhysicsAsset = SharedData->PhysicsAsset;
-	for(int32 BodySetupIndex = 0; BodySetupIndex < SharedData->SelectedBodies.Num(); ++BodySetupIndex)
-	{
-		USkeletalBodySetup* BodySetup = SharedData->PhysicsAsset->SkeletalBodySetups[SharedData->SelectedBodies[BodySetupIndex].Index];
-		if(BodySetup)
+	
+	for (const FPhysicsAssetEditorSharedData::FSelection& SelectedElement : SharedData->UniqueSelectionReferencingBodies())
+	{	
+		if(USkeletalBodySetup* const BodySetup = SharedData->PhysicsAsset->SkeletalBodySetups[SelectedElement.Index])
 		{
 			FName ProfileName = BodySetup->GetCurrentPhysicalAnimationProfileName();
 			BodySetup->RemovePhysicalAnimationProfile(ProfileName);
@@ -762,21 +767,22 @@ bool FPhysicsAssetDetailsCustomization::CanRemoveBodyFromPhysicalAnimationProfil
 
 	auto PhysicalAnimationProfileExistsForAny = [WeakSharedData]()
 	{
-		TSharedPtr<FPhysicsAssetEditorSharedData> LocalSharedData = WeakSharedData.Pin();
-
-		for(int32 BodySetupIndex = 0; BodySetupIndex < LocalSharedData->SelectedBodies.Num(); ++BodySetupIndex)
+		if (TSharedPtr<FPhysicsAssetEditorSharedData> LocalSharedData = WeakSharedData.Pin())
 		{
-			USkeletalBodySetup* BodySetup = LocalSharedData->PhysicsAsset->SkeletalBodySetups[LocalSharedData->SelectedBodies[BodySetupIndex].Index];
-			if(BodySetup && BodySetup->FindPhysicalAnimationProfile(BodySetup->GetCurrentPhysicalAnimationProfileName()))
+			for (const FPhysicsAssetEditorSharedData::FSelection& SelectedElement : LocalSharedData->UniqueSelectionReferencingBodies())
 			{
-				return true;
+				USkeletalBodySetup* const BodySetup = LocalSharedData->PhysicsAsset->SkeletalBodySetups[SelectedElement.Index];
+				if (BodySetup && BodySetup->FindPhysicalAnimationProfile(BodySetup->GetCurrentPhysicalAnimationProfileName()))
+				{
+					return true;
+				}
 			}
 		}
 
 		return false;
 	};
 
-	const bool bSelectedBodies = SharedData->SelectedBodies.Num() > 0;
+	const bool bSelectedBodies = !SharedData->UniqueSelectionReferencingBodies().IsEmpty();
 	return (PhysicsAssetEditorPtr.Pin()->IsNotSimulation() && bSelectedBodies && PhysicalAnimationProfileExistsForAny() && PhysicsAsset->CurrentPhysicalAnimationProfileName != NAME_None);
 }
 
@@ -795,8 +801,8 @@ void FPhysicsAssetDetailsCustomization::SelectAllBodiesInCurrentPhysicalAnimatio
 			NewBodiesSelection.Add(BSIndex);
 		}
 	}
-	SharedData->ClearSelectedBody();	//clear selection
-	SharedData->SetSelectedBodiesAnyPrimitive(NewBodiesSelection, true);
+
+	SharedData->SetSelectedBodies(NewBodiesSelection);
 
 	return;
 }
@@ -824,9 +830,10 @@ bool FPhysicsAssetDetailsCustomization::ConstraintProfileExistsForAny() const
 {
 	TSharedPtr<FPhysicsAssetEditorSharedData> SharedData = PhysicsAssetEditorPtr.Pin()->GetSharedData();
 	const FName ProfileName = SharedData->PhysicsAsset->CurrentConstraintProfileName;
-	for(int32 ConstraintIndex = 0; ConstraintIndex < SharedData->SelectedConstraints.Num(); ++ConstraintIndex)
+	
+	for(const FPhysicsAssetEditorSharedData::FSelection& SelectedConstraint : SharedData->SelectedConstraints())
 	{
-		UPhysicsConstraintTemplate* ConstraintSetup = SharedData->PhysicsAsset->ConstraintSetup[SharedData->SelectedConstraints[ConstraintIndex].Index];
+		const UPhysicsConstraintTemplate* const ConstraintSetup = SharedData->PhysicsAsset->ConstraintSetup[SelectedConstraint.Index];
 		if(ConstraintSetup && ConstraintSetup->ContainsConstraintProfile(ProfileName))
 		{
 			return true;
@@ -901,9 +908,10 @@ void FPhysicsAssetDetailsCustomization::AddConstraintToCurrentConstraintProfile(
 {
 	const FScopedTransaction Transaction(LOCTEXT("AssignToConstraintProfile", "Assign To Constraint Profile"));
 	TSharedPtr<FPhysicsAssetEditorSharedData> SharedData = PhysicsAssetEditorPtr.Pin()->GetSharedData();
-	for(int32 ConstraintIndex = 0; ConstraintIndex < SharedData->SelectedConstraints.Num(); ++ConstraintIndex)
-	{
-		UPhysicsConstraintTemplate* ConstraintSetup = SharedData->PhysicsAsset->ConstraintSetup[SharedData->SelectedConstraints[ConstraintIndex].Index];
+	
+	for (const FPhysicsAssetEditorSharedData::FSelection& SelectedConstraint : SharedData->SelectedConstraints())
+	{	
+		UPhysicsConstraintTemplate* const ConstraintSetup = SharedData->PhysicsAsset->ConstraintSetup[SelectedConstraint.Index];
 		FName ProfileName = ConstraintSetup->GetCurrentConstraintProfileName();
 		if (!ConstraintSetup->ContainsConstraintProfile(ProfileName))
 		{
@@ -921,29 +929,30 @@ bool FPhysicsAssetDetailsCustomization::CanAddConstraintToCurrentConstraintProfi
 
 	auto ConstraintProfileExistsForAll = [WeakSharedData]()
 	{
-		TSharedPtr<FPhysicsAssetEditorSharedData> LocalSharedData = WeakSharedData.Pin();
-
-		const FName ProfileName = LocalSharedData->PhysicsAsset->CurrentConstraintProfileName;
-		for(int32 ConstraintIndex = 0; ConstraintIndex < LocalSharedData->SelectedConstraints.Num(); ++ConstraintIndex)
+		if (TSharedPtr<FPhysicsAssetEditorSharedData> LocalSharedData = WeakSharedData.Pin())
 		{
-			UPhysicsConstraintTemplate* ConstraintSetup = LocalSharedData->PhysicsAsset->ConstraintSetup[LocalSharedData->SelectedConstraints[ConstraintIndex].Index];
-			if(ConstraintSetup)
+			const FName ProfileName = LocalSharedData->PhysicsAsset->CurrentConstraintProfileName;
+			for (const FPhysicsAssetEditorSharedData::FSelection& SelectedConstraint : LocalSharedData->SelectedConstraints())
 			{
-				if(!ConstraintSetup->ContainsConstraintProfile(ProfileName))
+				UPhysicsConstraintTemplate* ConstraintSetup = LocalSharedData->PhysicsAsset->ConstraintSetup[SelectedConstraint.Index];
+				if (ConstraintSetup)
+				{
+					if (!ConstraintSetup->ContainsConstraintProfile(ProfileName))
+					{
+						return false;
+					}
+				}
+				else
 				{
 					return false;
 				}
 			}
-			else
-			{
-				return false;
-			}
 		}
 
 		return true;
 	};
 
-	const bool bSelectedConstraints = SharedData->SelectedConstraints.Num() > 0;
+	const bool bSelectedConstraints = !SharedData->SelectedConstraints().IsEmpty();
 	return (PhysicsAssetEditorPtr.Pin()->IsNotSimulation() && bSelectedConstraints && PhysicsAsset->CurrentConstraintProfileName != NAME_None && !ConstraintProfileExistsForAll());
 }
 
@@ -951,9 +960,9 @@ void FPhysicsAssetDetailsCustomization::RemoveConstraintFromCurrentConstraintPro
 {
 	const FScopedTransaction Transaction(LOCTEXT("UnassignFromConstraintProfile", "Unassign From Constraint Profile"));
 	TSharedPtr<FPhysicsAssetEditorSharedData> SharedData = PhysicsAssetEditorPtr.Pin()->GetSharedData();
-	for(int32 ConstraintIndex = 0; ConstraintIndex < SharedData->SelectedConstraints.Num(); ++ConstraintIndex)
+	for (const FPhysicsAssetEditorSharedData::FSelection& SelectedConstraint : SharedData->SelectedConstraints())
 	{
-		UPhysicsConstraintTemplate* ConstraintSetup = SharedData->PhysicsAsset->ConstraintSetup[SharedData->SelectedConstraints[ConstraintIndex].Index];
+		UPhysicsConstraintTemplate* const ConstraintSetup = SharedData->PhysicsAsset->ConstraintSetup[SelectedConstraint.Index];
 
 		ConstraintSetup->Modify();
 		FName ProfileName = ConstraintSetup->GetCurrentConstraintProfileName();
@@ -969,22 +978,23 @@ bool FPhysicsAssetDetailsCustomization::CanRemoveConstraintFromCurrentConstraint
 
 	auto ConstraintProfileExistsForAny = [WeakSharedData]()
 	{
-		TSharedPtr<FPhysicsAssetEditorSharedData> LocalSharedData = WeakSharedData.Pin();
-
-		const FName ProfileName = LocalSharedData->PhysicsAsset->CurrentConstraintProfileName;
-		for(int32 ConstraintIndex = 0; ConstraintIndex < LocalSharedData->SelectedConstraints.Num(); ++ConstraintIndex)
+		if (TSharedPtr<FPhysicsAssetEditorSharedData> LocalSharedData = WeakSharedData.Pin())
 		{
-			UPhysicsConstraintTemplate* ConstraintSetup = LocalSharedData->PhysicsAsset->ConstraintSetup[LocalSharedData->SelectedConstraints[ConstraintIndex].Index];
-			if(ConstraintSetup && ConstraintSetup->ContainsConstraintProfile(ProfileName))
+			const FName ProfileName = LocalSharedData->PhysicsAsset->CurrentConstraintProfileName;
+			for (const FPhysicsAssetEditorSharedData::FSelection& SelectedConstraint : LocalSharedData->SelectedConstraints())
 			{
-				return true;
+				const UPhysicsConstraintTemplate* const ConstraintSetup = LocalSharedData->PhysicsAsset->ConstraintSetup[SelectedConstraint.Index];
+				if (ConstraintSetup && ConstraintSetup->ContainsConstraintProfile(ProfileName))
+				{
+					return true;
+				}
 			}
 		}
 
 		return false;
 	};
 
-	const bool bSelectedConstraints = SharedData->SelectedConstraints.Num() > 0;
+	const bool bSelectedConstraints = !SharedData->SelectedConstraints().IsEmpty();
 	return (PhysicsAssetEditorPtr.Pin()->IsNotSimulation() && bSelectedConstraints && PhysicsAsset->CurrentConstraintProfileName != NAME_None && ConstraintProfileExistsForAny());
 }
 
@@ -1005,7 +1015,7 @@ void FPhysicsAssetDetailsCustomization::SelectAllBodiesInCurrentConstraintProfil
 	}
 
 	SharedData->ClearSelectedConstraints();	//clear selection
-	SharedData->SetSelectedConstraints(NewSelectedConstraints, true);
+	SharedData->ModifySelectedConstraints(NewSelectedConstraints, true);
 
 	return;
 }
diff --git a/Engine/Source/Editor/PhysicsAssetEditor/Private/PhysicsAssetEditor.cpp b/Engine/Source/Editor/PhysicsAssetEditor/Private/PhysicsAssetEditor.cpp
--- a/Engine/Source/Editor/PhysicsAssetEditor/Private/PhysicsAssetEditor.cpp
+++ b/Engine/Source/Editor/PhysicsAssetEditor/Private/PhysicsAssetEditor.cpp
@@ -26,6 +26,7 @@
 #include "PhysicsAssetEditorModule.h"
 #include "ScopedTransaction.h"
 #include "PhysicsAssetEditorActions.h"
+#include "PhysicsAssetEditorSelection.h"
 #include "PhysicsAssetRenderUtils.h"
 #include "PhysicsAssetEditorSkeletalMeshComponent.h"
 #include "PhysicsAssetEditorToolMenuContext.h"
@@ -292,7 +293,7 @@ TSharedPtr<FPhysicsAssetEditorSharedData> FPhysicsAssetEditor::GetSharedData() c
 	return SharedData;
 }
 
-void FPhysicsAssetEditor::HandleViewportSelectionChanged(const TArray<FPhysicsAssetEditorSharedData::FSelection>& InSelectedBodies, const TArray<FPhysicsAssetEditorSharedData::FSelection>& InSelectedConstraints, const TArray<FPhysicsAssetEditorSharedData::FSelection>& InSelectedCoMs)
+void FPhysicsAssetEditor::HandleViewportSelectionChanged(const TArray<FPhysicsAssetEditorSelectedElement>& InSelectedElements)
 {
 	if (!bSelecting)
 	{
@@ -303,7 +304,7 @@ void FPhysicsAssetEditor::HandleViewportSelectionChanged(const TArray<FPhysicsAs
 			SkeletonTree->DeselectAll();
 		}
 
-		if(InSelectedBodies.IsEmpty() && InSelectedConstraints.IsEmpty() && InSelectedCoMs.IsEmpty())
+		if(InSelectedElements.IsEmpty())
 		{
 			if (PhysAssetProperties.IsValid())
 			{
@@ -324,36 +325,29 @@ void FPhysicsAssetEditor::HandleViewportSelectionChanged(const TArray<FPhysicsAs
 			TSet<USkeletalBodySetup*> Bodies;
 			TSet<UPhysicsConstraintTemplate*> Constraints;
 			TSet<PhysicsAssetEditor::FShapeData> Shapes;
-			Algo::Transform(InSelectedBodies, Objects, [this](const FPhysicsAssetEditorSharedData::FSelection& InItem) 
-			{ 
-				return SharedData->PhysicsAsset->SkeletalBodySetups[InItem.Index];
-			});
-			Algo::Transform(InSelectedConstraints, Objects, [this](const FPhysicsAssetEditorSharedData::FSelection& InItem) 
-			{ 
-				return SharedData->PhysicsAsset->ConstraintSetup[InItem.Index];
-			});
-			Algo::Transform(InSelectedCoMs, Objects, [this](const FPhysicsAssetEditorSharedData::FSelection& InItem)
-			{
-					return SharedData->PhysicsAsset->SkeletalBodySetups[InItem.Index]; // Add the owning physics body here so we display its details panel in the UI.
-			});
-			Algo::Transform(InSelectedBodies, Bodies, [this](const FPhysicsAssetEditorSharedData::FSelection& InItem) 
-			{ 
-				return SharedData->PhysicsAsset->SkeletalBodySetups[InItem.Index];
-			});
-			Algo::Transform(InSelectedConstraints, Constraints, [this](const FPhysicsAssetEditorSharedData::FSelection& InItem) 
-			{ 
-				return SharedData->PhysicsAsset->ConstraintSetup[InItem.Index];
-			});
-			// NOTE: Does selecting a constraint also select a shape?
-			// If not then this is not needed.
-			Algo::Transform(InSelectedConstraints, Shapes, [this](const FPhysicsAssetEditorSharedData::FSelection& InItem)
-			{
-				return PhysicsAssetEditor::FShapeData(InItem.Index, InItem.PrimitiveIndex, InItem.PrimitiveType);
-			});
-			Algo::Transform(InSelectedBodies, Shapes, [this](const FPhysicsAssetEditorSharedData::FSelection& InItem)
+			
+			for (const FPhysicsAssetEditorSelectedElement& SelectedElement : InSelectedElements)
 			{
-				return PhysicsAssetEditor::FShapeData(InItem.Index, InItem.PrimitiveIndex, InItem.PrimitiveType);
-			});
+				if (SelectedElement.HasType(FPhysicsAssetEditorSelectedElement::Body | FPhysicsAssetEditorSelectedElement::Primitive))
+				{
+					USkeletalBodySetup* const SelectedBodySetup = SharedData->PhysicsAsset->SkeletalBodySetups[SelectedElement.GetIndex()];
+					Bodies.Add(SelectedBodySetup);
+					Shapes.Add(PhysicsAssetEditor::FShapeData(SelectedElement.GetIndex(), SelectedElement.GetPrimitiveIndex(), SelectedElement.GetPrimitiveType()));
+					Objects.Add(SelectedBodySetup);
+				}
+				else if (SelectedElement.HasType(FPhysicsAssetEditorSelectedElement::Constraint))
+				{
+					UPhysicsConstraintTemplate* const SelectedConstraint = SharedData->PhysicsAsset->ConstraintSetup[SelectedElement.GetIndex()];
+					Constraints.Add(SelectedConstraint);
+					Shapes.Add(PhysicsAssetEditor::FShapeData(SelectedElement.GetIndex(), SelectedElement.GetPrimitiveIndex(), SelectedElement.GetPrimitiveType()));
+					Objects.Add(SelectedConstraint);
+				}
+				else if (SelectedElement.HasType(FPhysicsAssetEditorSelectedElement::CenterOfMass))
+				{
+					USkeletalBodySetup* const SelectedBodySetup = SharedData->PhysicsAsset->SkeletalBodySetups[SelectedElement.GetIndex()];
+					Objects.Add(SelectedBodySetup); // Add the owning physics body here so we display its details panel in the UI.
+				}	
+			}
 
 			if (PhysAssetProperties.IsValid())
 			{
@@ -556,15 +550,14 @@ void FPhysicsAssetEditor::OnFinishedChangingProperties(const FPropertyChangedEve
 	if (PropertyName == GET_MEMBER_NAME_CHECKED(UBodySetup, AggGeom))
 	{
 		// reselect all the bodies that were selected when the array changed, because selection keeps a primitive type that may have changed since
-		TArray<int32> SelectedBodyIndices;
-		for (const FPhysicsAssetEditorSharedData::FSelection& selectedBody : SharedData->SelectedBodies)
+		TArray<FPhysicsAssetEditorSharedData::FSelection> ReselectedPrimitives;
+		for (const FPhysicsAssetEditorSharedData::FSelection& SelectedPrimitive : SharedData->SelectedPrimitives())
 		{
-			SelectedBodyIndices.AddUnique(selectedBody.Index);
+			ReselectedPrimitives.AddUnique(MakeSelectionAnyPrimitiveInBody(SharedData->PhysicsAsset, SelectedPrimitive.Index));
 		}
 		{ // bulk update
 			FScopedBulkSelection BulkSelection(SharedData);
-			SharedData->ClearSelectedBody();
-			SharedData->SetSelectedBodiesAnyPrimitive(SelectedBodyIndices, true);
+			SharedData->SetSelectedPrimitives(ReselectedPrimitives);
 		}
 	}
 
@@ -1506,7 +1499,7 @@ void FPhysicsAssetEditor::BindCommands()
 		Commands.ToggleSelectionType,
 		FExecuteAction::CreateSP(this, &FPhysicsAssetEditor::OnToggleSelectionType, true),
 		FCanExecuteAction::CreateSP(this, &FPhysicsAssetEditor::IsNotSimulation));
-	
+
 	ToolkitCommands->MapAction(
 		Commands.ToggleSelectionTypeWithUserConstraints,
 		FExecuteAction::CreateSP(this, &FPhysicsAssetEditor::OnToggleSelectionType, false),
@@ -2196,18 +2189,23 @@ TSharedRef<ISkeletonTree> FPhysicsAssetEditor::BuildMenuWidgetNewConstraintForBo
 			{
 				HandleCreateNewConstraint(InSourceBodyIndex, SelectedBody->GetBodySetupIndex());
 			}
-			else if(SharedData->GetSelectedBody() != nullptr)
+			else
 			{
-				// make a copy to avoid changing SelectedBodies while iterating SelectedBodies
-				TArray<int32> SourceBodyIndices;
-				for (const FPhysicsAssetEditorSharedData::FSelection& SourceBody : SharedData->SelectedBodies)
-				{
-					SourceBodyIndices.Add(SourceBody.Index);
-				}
-				// create constraints
-				for(const int32 SourceBodyIndex : SourceBodyIndices)
+				const FPhysicsAssetEditorSharedData::SelectionUniqueRange SelectionRange = SharedData->UniqueSelectionReferencingBodies();
+
+				if (!SelectionRange.IsEmpty())
 				{
-					HandleCreateNewConstraint(SourceBodyIndex, SelectedBody->GetBodySetupIndex());
+					// make a copy to avoid changing SelectedBodies while iterating SelectedBodies
+					TArray<int32> SourceBodyIndices;
+					for (const FPhysicsAssetEditorSharedData::FSelection& SourceBody : SelectionRange)
+					{
+						SourceBodyIndices.Add(SourceBody.Index);
+					}
+					// create constraints
+					for (const int32 SourceBodyIndex : SourceBodyIndices)
+					{
+						HandleCreateNewConstraint(SourceBodyIndex, SelectedBody->GetBodySetupIndex());
+					}
 				}
 			}
 		}
@@ -2283,9 +2281,9 @@ void FPhysicsAssetEditor::CreateOrConvertConstraint(EPhysicsAssetEditorConstrain
 
 	const FScopedTransaction Transaction( LOCTEXT( "CreateConvertConstraint", "Create Or Convert Constraint" ) );
 
-	for(int32 i=0; i<SharedData->SelectedConstraints.Num(); ++i)
+	for (const FPhysicsAssetEditorSharedData::FSelection& SelectedConstraint : SharedData->SelectedConstraints())
 	{
-		UPhysicsConstraintTemplate* ConstraintSetup = SharedData->PhysicsAsset->ConstraintSetup[SharedData->SelectedConstraints[i].Index];
+		UPhysicsConstraintTemplate* const ConstraintSetup = SharedData->PhysicsAsset->ConstraintSetup[SelectedConstraint.Index];
 		ConstraintSetup->PreEditChange(DefaultInstanceProperty);
 
 		if(ConstraintType == EPCT_BSJoint)
@@ -2316,22 +2314,18 @@ void FPhysicsAssetEditor::CreateOrConvertConstraint(EPhysicsAssetEditorConstrain
 
 void FPhysicsAssetEditor::AddNewPrimitive(EAggCollisionShape::Type InPrimitiveType, bool bCopySelected)
 {
-	check(!bCopySelected || SharedData->SelectedBodies.Num() == 1);	//we only support this for one selection
+	TArray<FPhysicsAssetEditorSharedData::FSelection> NewSelection = SharedData->UniqueSelectionReferencingBodies().ToArray();
+
+	check(!bCopySelected || NewSelection.Num() == 1);	//we only support this for one selection
 	int32 NewPrimIndex = 0;
-	TArray<FPhysicsAssetEditorSharedData::FSelection> NewSelection;
+	
 	{
 		// Make sure rendering is done - so we are not changing data being used by collision drawing.
 		FlushRenderingCommands();
 
 		const FScopedTransaction Transaction( NSLOCTEXT("UnrealEd", "AddNewPrimitive", "Add New Primitive") );
 
-		//first we need to grab all the bodies we're modifying (removes duplicates from multiple primitives)
-		for(int32 i=0; i<SharedData->SelectedBodies.Num(); ++i)
-		{
-			NewSelection.AddUnique(FPhysicsAssetEditorSharedData::FSelection(SharedData->SelectedBodies[i].Index, EAggCollisionShape::Unknown, 0));	//only care about body index for now, we'll later update the primitive index
-		}
-
-		// Make new bodies for any bones we have selected that dont already have them
+		// Make new bodies for any bones we have selected that don't already have them.
 		TArray<TSharedPtr<ISkeletonTreeItem>> Items = SkeletonTree->GetSelectedItems();
 		FSkeletonTreeSelection Selection(Items);
 		TArray<TSharedPtr<ISkeletonTreeItem>> BoneItems = Selection.GetSelectedItemsByTypeId("FSkeletonTreeBoneItem");
@@ -2345,18 +2339,18 @@ void FPhysicsAssetEditor::AddNewPrimitive(EAggCollisionShape::Type InPrimitiveTy
 			{
 				const FPhysAssetCreateParams& NewBodyData = GetDefault<UPhysicsAssetGenerationSettings>()->CreateParams;
 				int32 NewBodyIndex = FPhysicsAssetUtils::CreateNewBody(SharedData->PhysicsAsset, BoneProxy->BoneName, NewBodyData);
-				NewSelection.AddUnique(FPhysicsAssetEditorSharedData::FSelection(NewBodyIndex, EAggCollisionShape::Unknown, 0));
+				NewSelection.AddUnique(MakePrimitiveSelection(NewBodyIndex, EAggCollisionShape::Unknown, 0));
 			}
 		}
 
 		for(int32 i=0; i<NewSelection.Num(); ++i)
 		{
-			int32 BodyIndex = NewSelection[i].Index;
+			const int32 BodyIndex = NewSelection[i].Index;
 			UBodySetup* BodySetup = SharedData->PhysicsAsset->SkeletalBodySetups[BodyIndex];
 			EAggCollisionShape::Type PrimitiveType;
 			if (bCopySelected)
 			{
-				PrimitiveType = SharedData->GetSelectedBody()->PrimitiveType;
+				PrimitiveType = SharedData->GetSelectedBodyOrPrimitive()->GetPrimitiveType();
 			}
 			else
 			{
@@ -2380,10 +2374,10 @@ void FPhysicsAssetEditor::AddNewPrimitive(EAggCollisionShape::Type InPrimitiveTy
 				}
 				else
 				{
-					SphereElem->Center = BodySetup->AggGeom.SphereElems[SharedData->GetSelectedBody()->PrimitiveIndex].Center;
+					SphereElem->Center = BodySetup->AggGeom.SphereElems[SharedData->GetSelectedBodyOrPrimitive()->PrimitiveIndex].Center;
 					SphereElem->Center.X += PhysicsAssetEditor::DuplicateXOffset;
 
-					SphereElem->Radius = BodySetup->AggGeom.SphereElems[SharedData->GetSelectedBody()->PrimitiveIndex].Radius;
+					SphereElem->Radius = BodySetup->AggGeom.SphereElems[SharedData->GetSelectedBodyOrPrimitive()->PrimitiveIndex].Radius;
 				}
 				SharedData->AutoNamePrimitive(BodyIndex, PrimitiveType);
 			}
@@ -2404,12 +2398,12 @@ void FPhysicsAssetEditor::AddNewPrimitive(EAggCollisionShape::Type InPrimitiveTy
 				}
 				else
 				{
-					BoxElem->SetTransform( BodySetup->AggGeom.BoxElems[SharedData->GetSelectedBody()->PrimitiveIndex].GetTransform() );
+					BoxElem->SetTransform( BodySetup->AggGeom.BoxElems[SharedData->GetSelectedBodyOrPrimitive()->PrimitiveIndex].GetTransform() );
 					BoxElem->Center.X += PhysicsAssetEditor::DuplicateXOffset;
 
-					BoxElem->X = BodySetup->AggGeom.BoxElems[SharedData->GetSelectedBody()->PrimitiveIndex].X;
-					BoxElem->Y = BodySetup->AggGeom.BoxElems[SharedData->GetSelectedBody()->PrimitiveIndex].Y;
-					BoxElem->Z = BodySetup->AggGeom.BoxElems[SharedData->GetSelectedBody()->PrimitiveIndex].Z;
+					BoxElem->X = BodySetup->AggGeom.BoxElems[SharedData->GetSelectedBodyOrPrimitive()->PrimitiveIndex].X;
+					BoxElem->Y = BodySetup->AggGeom.BoxElems[SharedData->GetSelectedBodyOrPrimitive()->PrimitiveIndex].Y;
+					BoxElem->Z = BodySetup->AggGeom.BoxElems[SharedData->GetSelectedBodyOrPrimitive()->PrimitiveIndex].Z;
 				}
 				SharedData->AutoNamePrimitive(BodyIndex, PrimitiveType);
 			}
@@ -2429,11 +2423,11 @@ void FPhysicsAssetEditor::AddNewPrimitive(EAggCollisionShape::Type InPrimitiveTy
 				}
 				else
 				{
-					SphylElem->SetTransform( BodySetup->AggGeom.SphylElems[SharedData->GetSelectedBody()->PrimitiveIndex].GetTransform() );
+					SphylElem->SetTransform( BodySetup->AggGeom.SphylElems[SharedData->GetSelectedBodyOrPrimitive()->PrimitiveIndex].GetTransform() );
 					SphylElem->Center.X += PhysicsAssetEditor::DuplicateXOffset;
 
-					SphylElem->Length = BodySetup->AggGeom.SphylElems[SharedData->GetSelectedBody()->PrimitiveIndex].Length;
-					SphylElem->Radius = BodySetup->AggGeom.SphylElems[SharedData->GetSelectedBody()->PrimitiveIndex].Radius;
+					SphylElem->Length = BodySetup->AggGeom.SphylElems[SharedData->GetSelectedBodyOrPrimitive()->PrimitiveIndex].Length;
+					SphylElem->Radius = BodySetup->AggGeom.SphylElems[SharedData->GetSelectedBodyOrPrimitive()->PrimitiveIndex].Radius;
 				}
 				SharedData->AutoNamePrimitive(BodyIndex, PrimitiveType);
 			}
@@ -2446,10 +2440,10 @@ void FPhysicsAssetEditor::AddNewPrimitive(EAggCollisionShape::Type InPrimitiveTy
 				NewSelection[i].PrimitiveIndex = NewPrimIndex;
 				FKConvexElem* ConvexElem = &BodySetup->AggGeom.ConvexElems[NewPrimIndex];
 
-				ConvexElem->SetTransform(BodySetup->AggGeom.ConvexElems[SharedData->GetSelectedBody()->PrimitiveIndex].GetTransform());
+				ConvexElem->SetTransform(BodySetup->AggGeom.ConvexElems[SharedData->GetSelectedBodyOrPrimitive()->PrimitiveIndex].GetTransform());
 
 				// Copy all of the vertices of the convex element
-				for (FVector v : BodySetup->AggGeom.ConvexElems[SharedData->GetSelectedBody()->PrimitiveIndex].VertexData)
+				for (FVector v : BodySetup->AggGeom.ConvexElems[SharedData->GetSelectedBodyOrPrimitive()->PrimitiveIndex].VertexData)
 				{
 					v.X += PhysicsAssetEditor::DuplicateXOffset;
 					ConvexElem->VertexData.Add(v);
@@ -2478,12 +2472,12 @@ void FPhysicsAssetEditor::AddNewPrimitive(EAggCollisionShape::Type InPrimitiveTy
 				}
 				else
 				{
-					TaperedCapsuleElem->SetTransform( BodySetup->AggGeom.TaperedCapsuleElems[SharedData->GetSelectedBody()->PrimitiveIndex].GetTransform() );
+					TaperedCapsuleElem->SetTransform( BodySetup->AggGeom.TaperedCapsuleElems[SharedData->GetSelectedBodyOrPrimitive()->PrimitiveIndex].GetTransform() );
 					TaperedCapsuleElem->Center.X += PhysicsAssetEditor::DuplicateXOffset;
 
-					TaperedCapsuleElem->Length = BodySetup->AggGeom.TaperedCapsuleElems[SharedData->GetSelectedBody()->PrimitiveIndex].Length;
-					TaperedCapsuleElem->Radius0 = BodySetup->AggGeom.TaperedCapsuleElems[SharedData->GetSelectedBody()->PrimitiveIndex].Radius0;
-					TaperedCapsuleElem->Radius1 = BodySetup->AggGeom.TaperedCapsuleElems[SharedData->GetSelectedBody()->PrimitiveIndex].Radius1;
+					TaperedCapsuleElem->Length = BodySetup->AggGeom.TaperedCapsuleElems[SharedData->GetSelectedBodyOrPrimitive()->PrimitiveIndex].Length;
+					TaperedCapsuleElem->Radius0 = BodySetup->AggGeom.TaperedCapsuleElems[SharedData->GetSelectedBodyOrPrimitive()->PrimitiveIndex].Radius0;
+					TaperedCapsuleElem->Radius1 = BodySetup->AggGeom.TaperedCapsuleElems[SharedData->GetSelectedBodyOrPrimitive()->PrimitiveIndex].Radius1;
 				}
 
 				SharedData->AutoNamePrimitive(BodyIndex, PrimitiveType);
@@ -2496,8 +2490,7 @@ void FPhysicsAssetEditor::AddNewPrimitive(EAggCollisionShape::Type InPrimitiveTy
 	} // ScopedTransaction
 
 	//clear selection
-	SharedData->ClearSelectedBody();
-	SharedData->SetSelectedBodies(NewSelection, true);
+	SharedData->SetSelectedPrimitives(NewSelection);
 
 	RecreatePhysicsState();
 	RefreshHierachyTree();
@@ -2507,9 +2500,9 @@ void FPhysicsAssetEditor::AddNewPrimitive(EAggCollisionShape::Type InPrimitiveTy
 void FPhysicsAssetEditor::SetBodiesBelowSelectedPhysicsType( EPhysicsType InPhysicsType, bool bMarkAsDirty)
 {
 	TArray<int32> Indices;
-	for(int32 i=0; i<SharedData->SelectedBodies.Num(); ++i)
+	for (const FPhysicsAssetEditorSharedData::FSelection& SelectedBody : SharedData->UniqueSelectionReferencingBodies())
 	{
-		Indices.Add(SharedData->SelectedBodies[i].Index);
+		Indices.Add(SelectedBody.Index);
 	}
 
 	SetBodiesBelowPhysicsType(InPhysicsType, Indices, bMarkAsDirty);
@@ -2558,22 +2551,22 @@ bool FPhysicsAssetEditor::IsNotSimulation() const
 
 bool FPhysicsAssetEditor::HasSelectedBodyAndIsNotSimulation() const
 {
-	return IsNotSimulation() && (SharedData->GetSelectedBody());
+	return IsNotSimulation() && (SharedData->GetSelectedBodyOrPrimitive());
 }
 
 bool FPhysicsAssetEditor::HasOneSelectedBodyAndIsNotSimulation() const
 {
-	return IsNotSimulation() && (SharedData->SelectedBodies.Num() == 1);
+	return IsNotSimulation() && (SharedData->UniqueSelectionReferencingBodies().Num() == 1);
 }
 
 bool FPhysicsAssetEditor::HasMoreThanOneSelectedBodyAndIsNotSimulation() const
 {
-	return IsNotSimulation() && (SharedData->SelectedBodies.Num() > 1);
+	return IsNotSimulation() && (SharedData->UniqueSelectionReferencingBodies().Num() > 1);
 }
 
 bool FPhysicsAssetEditor::HasSelectedBodyOrConstraintAndIsNotSimulation() const
 {
-	return IsNotSimulation() && (SharedData->SelectedBodies.Num() > 0 || SharedData->SelectedConstraints.Num() > 0);
+	return IsNotSimulation() && (!SharedData->UniqueSelectionReferencingBodies().IsEmpty() || !SharedData->SelectedConstraints().IsEmpty());
 }
 
 bool FPhysicsAssetEditor::CanEditConstraintProperties() const
@@ -2583,7 +2576,7 @@ bool FPhysicsAssetEditor::CanEditConstraintProperties() const
 		//If we are currently editing a constraint profile, make sure all selected constraints belong to the profile
 		if(SharedData->PhysicsAsset->CurrentConstraintProfileName != NAME_None)
 		{
-			for (const FPhysicsAssetEditorSharedData::FSelection& Selection : SharedData->SelectedConstraints)
+			for (const FPhysicsAssetEditorSharedData::FSelection& Selection : SharedData->SelectedConstraints())
 			{
 				UPhysicsConstraintTemplate* CS = SharedData->PhysicsAsset->ConstraintSetup[Selection.Index];
 				if(!CS || !CS->ContainsConstraintProfile(SharedData->PhysicsAsset->CurrentConstraintProfileName))
@@ -2608,7 +2601,7 @@ bool FPhysicsAssetEditor::HasSelectedConstraintAndIsNotSimulation() const
 
 void FPhysicsAssetEditor::OnCopyBodyName()
 {
-	if(SharedData->SelectedBodies.Num() == 1)
+	if(SharedData->UniqueSelectionReferencingBodies().Num() == 1)
 	{
 		SharedData->CopyBodyName();
 	}
@@ -2618,7 +2611,7 @@ void FPhysicsAssetEditor::OnCopyBodyName()
 
 bool FPhysicsAssetEditor::CanCopyBodyName() const
 {
-	return IsSelectedEditMode() && SharedData->SelectedBodies.Num() == 1 && SharedData->SelectedConstraints.Num() == 0;
+	return IsSelectedEditMode() && SharedData->UniqueSelectionReferencingBodies().Num() == 1 && SharedData->SelectedConstraints().IsEmpty();
 }
 
 bool FPhysicsAssetEditor::IsSelectedEditMode() const
@@ -2666,13 +2659,14 @@ void FPhysicsAssetEditor::CreateBodiesAndConstraintsForSelectedBones(const FPhys
 	// Make sure rendering is done - so we are not changing data being used by collision drawing.
 	FlushRenderingCommands();
 
-	if(SharedData->SelectedBodies.Num() > 0)
+	if(!SharedData->UniqueSelectionReferencingBodies().IsEmpty())
 	{
 		TArray<int32> SelectedBodyIndices;
 		const FScopedTransaction Transaction( LOCTEXT("ResetBoneCollision", "Reset Bone Collision") );
-		for (int32 i = 0; i < SharedData->SelectedBodies.Num(); ++i)
+
+		for (const FPhysicsAssetEditorSharedData::FSelection& SelectedBody : SharedData->UniqueSelectionReferencingBodies())
 		{
-			int32 SelectedBodyIndex = SharedData->SelectedBodies[i].Index;
+			int32 SelectedBodyIndex = SelectedBody.Index;
 			if (SharedData->PhysicsAsset->SkeletalBodySetups.IsValidIndex(SelectedBodyIndex) == false)
 			{
 				continue;
@@ -2690,9 +2684,7 @@ void FPhysicsAssetEditor::CreateBodiesAndConstraintsForSelectedBones(const FPhys
 			SharedData->AutoNameAllPrimitives(BodyIndex, NewBodyData.GeomType);
 		}
 
-		//deselect first
-		SharedData->ClearSelectedBody();
-		SharedData->SetSelectedBodiesAnyPrimitive(BodyIndices, true);
+		SharedData->SetSelectedBodies(BodyIndices);
 	}
 	else
 	{
@@ -2790,7 +2782,7 @@ bool FPhysicsAssetEditor::CanCopyBodies() const
 {
 	if (IsSelectedEditMode())
 	{
-		return ((SharedData->SelectedBodies.Num() > 0) || (SharedData->SelectedConstraints.Num() > 0));
+		return !SharedData->UniqueSelectionReferencingBodies().IsEmpty() || !SharedData->SelectedConstraints().IsEmpty();
 	}
 	return false;
 }
@@ -2824,7 +2816,7 @@ bool FPhysicsAssetEditor::CanCopyShapes() const
 {
 	if (IsSelectedEditMode())
 	{
-		return (SharedData->SelectedBodies.Num() > 0);
+		return !SharedData->UniqueSelectionReferencingBodies().IsEmpty();
 	}
 	return false;
 }
@@ -2845,11 +2837,11 @@ bool FPhysicsAssetEditor::CanPasteShapes() const
 
 void FPhysicsAssetEditor::OnCopyProperties()
 {
-	if(SharedData->SelectedBodies.Num() == 1)
+	if(SharedData->UniqueSelectionReferencingBodies().Num() == 1)
 	{
 		SharedData->CopyBodyProperties();
 	}
-	else if(SharedData->SelectedConstraints.Num() == 1)
+	else if(SharedData->SelectedConstraints().Num() == 1)
 	{
 		SharedData->CopyConstraintProperties();
 	}
@@ -2859,11 +2851,11 @@ void FPhysicsAssetEditor::OnCopyProperties()
 
 void FPhysicsAssetEditor::OnPasteProperties()
 {
-	if(SharedData->SelectedBodies.Num() > 0)
+	if(!SharedData->UniqueSelectionReferencingBodies().IsEmpty())
 	{
 		SharedData->PasteBodyProperties();
 	}
-	else if (SharedData->SelectedConstraints.Num() > 0)
+	else if (!SharedData->SelectedConstraints().IsEmpty())
 	{
 		SharedData->PasteConstraintProperties();
 	}
@@ -2878,11 +2870,11 @@ bool FPhysicsAssetEditor::CanCopyProperties() const
 {
 	if(IsSelectedEditMode())
 	{
-		if(SharedData->SelectedBodies.Num() == 1 && SharedData->SelectedConstraints.Num() == 0)
+		if(SharedData->UniqueSelectionReferencingBodies().Num() == 1 && SharedData->SelectedConstraints().IsEmpty())
 		{
 			return true;
 		}
-		else if(SharedData->SelectedConstraints.Num() == 1 && SharedData->SelectedBodies.Num() == 0)
+		else if (SharedData->SelectedConstraints().Num() == 1 && SharedData->UniqueSelectionReferencingBodies().IsEmpty())
 		{
 			return true;
 		}
@@ -2893,7 +2885,7 @@ bool FPhysicsAssetEditor::CanCopyProperties() const
 
 bool FPhysicsAssetEditor::CanPasteProperties() const
 {
-	return IsSelectedEditMode() && IsCopyProperties() && (SharedData->SelectedBodies.Num() > 0 || SharedData->SelectedConstraints.Num() > 0);
+	return IsSelectedEditMode() && IsCopyProperties() && (!SharedData->UniqueSelectionReferencingBodies().IsEmpty() || !SharedData->SelectedConstraints().IsEmpty());
 }
 
 bool FPhysicsAssetEditor::IsCopyProperties() const
@@ -3033,10 +3025,10 @@ void FPhysicsAssetEditor::SetupSelectedSimulation()
 		//constraints need some more work
 		TArray<int32> BodyIndices;
 		TArray<TObjectPtr<UPhysicsConstraintTemplate>> & ConstraintSetup = PhysicsAsset->ConstraintSetup;
-		for(int32 i=0; i<SharedData->SelectedConstraints.Num(); ++i)
+		for (const FPhysicsAssetEditorSharedData::FSelection& SelectedConstraint : SharedData->SelectedConstraints())
 		{
-			int32 ConstraintIndex = SharedData->SelectedConstraints[i].Index;
-			FName ConstraintBone1 = ConstraintSetup[ConstraintIndex]->DefaultInstance.ConstraintBone1;	//we only unfix the child bodies
+			const int32 ConstraintIndex = SelectedConstraint.Index;
+			const FName ConstraintBone1 = ConstraintSetup[ConstraintIndex]->DefaultInstance.ConstraintBone1;	//we only unfix the child bodies
 
 			for(int32 j=0; j<BodySetup.Num(); ++j)
 			{
@@ -3353,16 +3345,18 @@ void FPhysicsAssetEditor::OnDuplicatePrimitive()
 
 bool FPhysicsAssetEditor::CanDuplicatePrimitive() const
 {
-	return HasSelectedBodyAndIsNotSimulation() && SharedData->SelectedBodies.Num() == 1;
+	return HasSelectedBodyAndIsNotSimulation() && SharedData->UniqueSelectionReferencingBodies().Num() == 1;
 }
 
 void FPhysicsAssetEditor::OnConstrainChildBodiesToParentBody()
 {
-	if (SharedData->SelectedBodies.Num() > 1)
+	const FPhysicsAssetEditorSharedData::FSelection* const LastSelectedBody = SharedData->GetSelectedBody();
+
+	if (LastSelectedBody && (SharedData->UniqueSelectionReferencingBodies().Num() > 1))
 	{
-		int32 ParentBodyIndex = SharedData->SelectedBodies.Last().Index;
+		const int32 ParentBodyIndex = LastSelectedBody->Index;
 		TArray<int32> ChildBodyIndices; // needed as the selection may contain multiple time the same body with different primitive index
-		for (const FPhysicsAssetEditorSharedData::FSelection& Selection : SharedData->SelectedBodies)
+		for (const FPhysicsAssetEditorSharedData::FSelection& Selection : SharedData->UniqueSelectionReferencingBodies())
 		{
 			if (Selection.Index != ParentBodyIndex)
 			{
@@ -3383,9 +3377,9 @@ void FPhysicsAssetEditor::OnSnapConstraint(const EConstraintTransformComponentFl
 {
 	const FScopedTransaction Transaction( LOCTEXT( "SnapConstraints", "Snap Constraints" ) );
 
-	for(int32 i=0; i<SharedData->SelectedConstraints.Num(); ++i)
+	for (const FPhysicsAssetEditorSharedData::FSelection& SelectedConstraint : SharedData->SelectedConstraints())
 	{
-		SharedData->SnapConstraintToBone(SharedData->SelectedConstraints[i].Index, ComponentFlags);
+		SharedData->SnapConstraintToBone(SelectedConstraint.Index, ComponentFlags);
 	}
 	
 	RefreshPreviewViewport();
@@ -3419,11 +3413,13 @@ void FPhysicsAssetEditor::OnDeleteConstraint()
 
 void FPhysicsAssetEditor::OnSetBodyPhysicsType( EPhysicsType InPhysicsType )
 {
-	if (SharedData->GetSelectedBody())
+	FPhysicsAssetEditorSharedData::SelectionUniqueRange SelectionRange = SharedData->UniqueSelectionReferencingBodies();
+
+	if (!SelectionRange.IsEmpty())
 	{
-		for(int32 i=0; i<SharedData->SelectedBodies.Num(); ++i)
+		for (const FPhysicsAssetEditorSharedData::FSelection& SelectedElement : SelectionRange)
 		{
-			UBodySetup* BodySetup = SharedData->PhysicsAsset->SkeletalBodySetups[SharedData->SelectedBodies[i].Index];
+			UBodySetup* const BodySetup = SharedData->PhysicsAsset->SkeletalBodySetups[SelectedElement.Index];
 			BodySetup->Modify();
 			BodySetup->PhysicsType = InPhysicsType;
 		}
@@ -3434,10 +3430,10 @@ void FPhysicsAssetEditor::OnSetBodyPhysicsType( EPhysicsType InPhysicsType )
 }
 
 bool FPhysicsAssetEditor::IsBodyPhysicsType( EPhysicsType InPhysicsType )
-{
-	for(int32 i=0; i<SharedData->SelectedBodies.Num(); ++i)
+{	
+	for (const FPhysicsAssetEditorSharedData::FSelection& SelectedBody : SharedData->UniqueSelectionReferencingBodies())
 	{
-		UBodySetup* BodySetup = SharedData->PhysicsAsset->SkeletalBodySetups[SharedData->SelectedBodies[i].Index];
+		UBodySetup* const BodySetup = SharedData->PhysicsAsset->SkeletalBodySetups[SelectedBody.Index];
 		if(BodySetup->PhysicsType == InPhysicsType)
 		{
 			return true;
@@ -3450,15 +3446,17 @@ bool FPhysicsAssetEditor::IsBodyPhysicsType( EPhysicsType InPhysicsType )
 
 void FPhysicsAssetEditor::OnDeleteBody()
 {
-	if(SharedData->SelectedBodies.Num())
+	FPhysicsAssetEditorSharedData::SelectionUniqueRange SelectionRange = SharedData->UniqueSelectionReferencingBodies();
+
+	if(!SelectionRange.IsEmpty())
 	{
 		//first build the bodysetup array because deleting bodies modifies the selected array
 		TArray<UBodySetup*> BodySetups;
-		BodySetups.Reserve(SharedData->SelectedBodies.Num());
-
-		for(int32 i=0; i<SharedData->SelectedBodies.Num(); ++i)
+		BodySetups.Reserve(SelectionRange.Num());
+	
+		for (const FPhysicsAssetEditorSharedData::FSelection& SelectedBody : SelectionRange)
 		{
-			BodySetups.Add( SharedData->PhysicsAsset->SkeletalBodySetups[SharedData->SelectedBodies[i].Index] );
+			BodySetups.Add( SharedData->PhysicsAsset->SkeletalBodySetups[SelectedBody.Index] );
 		}
 
 		const FScopedTransaction Transaction( LOCTEXT( "DeleteBodies", "Delete Bodies" ) );
@@ -3487,7 +3485,7 @@ void FPhysicsAssetEditor::OnDeleteAllBodiesBelow()
 
 	TArray<UBodySetup*> BodySetups;
 
-	for (FPhysicsAssetEditorSharedData::FSelection SelectedBody : SharedData->SelectedBodies)
+	for (const FPhysicsAssetEditorSharedData::FSelection& SelectedBody : SharedData->UniqueSelectionReferencingBodies())
 	{
 		UBodySetup* BaseSetup = SharedData->PhysicsAsset->SkeletalBodySetups[SelectedBody.Index];
 		
@@ -3622,11 +3620,11 @@ void FPhysicsAssetEditor::OnAssetSelectedFromStaticMeshAssetPicker( const FAsset
 	// gather all the body indices from both the body and bone selection 
 	// make sure to create a body setup if we encounter a bone with no associated body
 	TSet<int32> BodyIndicesToUpdate;
-	if (SharedData->GetSelectedBody() || BoneItems.Num() > 0)
+	if (SharedData->GetSelectedBodyOrPrimitive() || BoneItems.Num() > 0)
 	{
-		for (int32 SelectedBodyIndex = 0; SelectedBodyIndex < SharedData->SelectedBodies.Num(); ++SelectedBodyIndex)
+		for (const FPhysicsAssetEditorSharedData::FSelection& SelectedBody : SharedData->UniqueSelectionReferencingBodies())
 		{
-			BodyIndicesToUpdate.Add(SharedData->SelectedBodies[SelectedBodyIndex].Index);
+			BodyIndicesToUpdate.Add(SelectedBody.Index);
 		}
 
 		for (TSharedPtr<ISkeletonTreeItem> BoneItem : BoneItems)
@@ -3703,12 +3701,14 @@ TSharedRef<SWidget> FPhysicsAssetEditor::BuildPhysicalMaterialAssetPicker(bool b
 	}
 	else
 	{
-		if(SharedData->SelectedBodies.Num())
+		if(UPhysicsAssetEditorSelection::UniqueIterator SelectedBodyItr = SharedData->UniqueSelectionReferencingBodies().CreateConstIterator())
 		{
-			SelectedPhysicalMaterial = SharedData->PhysicsAsset->SkeletalBodySetups[SharedData->SelectedBodies[0].Index]->PhysMaterial;
-			for (int32 SelectedBodyIndex = 0; SelectedBodyIndex < SharedData->SelectedBodies.Num(); ++SelectedBodyIndex)
+			SelectedPhysicalMaterial = SharedData->PhysicsAsset->SkeletalBodySetups[SelectedBodyItr->Index]->PhysMaterial;
+			
+			for ( ; SelectedBodyItr; ++SelectedBodyItr)
 			{
-				USkeletalBodySetup* BodySetup = SharedData->PhysicsAsset->SkeletalBodySetups[SharedData->SelectedBodies[SelectedBodyIndex].Index];
+				const FPhysicsAssetEditorSharedData::FSelection& SelectedBody = *SelectedBodyItr;
+				const USkeletalBodySetup* const BodySetup = SharedData->PhysicsAsset->SkeletalBodySetups[SelectedBody.Index];
 				if(BodySetup->PhysMaterial != SelectedPhysicalMaterial)
 				{
 					SelectedPhysicalMaterial = nullptr;
@@ -3733,7 +3733,7 @@ void FPhysicsAssetEditor::OnAssetSelectedFromPhysicalMaterialAssetPicker( const
 {
 	FSlateApplication::Get().DismissAllMenus();
 
-	if (SharedData->GetSelectedBody() || bForAllBodies)
+	if (SharedData->GetSelectedBodyOrPrimitive() || bForAllBodies)
 	{
 		const FScopedTransaction Transaction(LOCTEXT("SetPhysicalMaterial", "Set Physical Material"));
 
@@ -3751,9 +3751,9 @@ void FPhysicsAssetEditor::OnAssetSelectedFromPhysicalMaterialAssetPicker( const
 			}
 			else
 			{
-				for (int32 SelectedBodyIndex = 0; SelectedBodyIndex < SharedData->SelectedBodies.Num(); ++SelectedBodyIndex)
+				for (const FPhysicsAssetEditorSharedData::FSelection& SelectedBody : SharedData->UniqueSelectionReferencingBodies())
 				{
-					USkeletalBodySetup* BodySetup = SharedData->PhysicsAsset->SkeletalBodySetups[SharedData->SelectedBodies[SelectedBodyIndex].Index];
+					USkeletalBodySetup* BodySetup = SharedData->PhysicsAsset->SkeletalBodySetups[SelectedBody.Index];
 					BodySetup->Modify();
 					BodySetup->PhysMaterial = PhysicalMaterial;
 				}
@@ -3819,7 +3819,7 @@ void FPhysicsAssetEditor::OnSelectShapes(const ECollisionEnabled::Type Collision
 {
 	UPhysicsAsset * const PhysicsAsset = SharedData->EditorSkelComp->GetPhysicsAsset();
 	TSet<int32> SelectedBodyIndices;
-	for (const FPhysicsAssetEditorSharedData::FSelection& SelectedBody : SharedData->SelectedBodies)
+	for (const FPhysicsAssetEditorSharedData::FSelection& SelectedBody : SharedData->UniqueSelectionReferencingBodies())
 	{
 		SelectedBodyIndices.Add(SelectedBody.Index);
 	}
@@ -4122,7 +4122,7 @@ void FPhysicsAssetEditor::GenerateSkinnedTriangleMesh()
 		else
 		{
 			SharedData->AutoNameAllPrimitives(BodyIndex, EFG_SkinnedTriangleMesh);
-			SharedData->SetSelectedBodyAnyPrimitive(BodyIndex, true);
+			SharedData->ModifySelectedBodies(BodyIndex, true);
 		}
 	}
 	RecreatePhysicsState();
@@ -4200,7 +4200,7 @@ void FPhysicsAssetEditor::OnSelectAllConstraints()
 
 	//Deselect everything first
 	SharedData->ClearSelectedConstraints();
-	SharedData->SetSelectedConstraints(NewSelectedConstraints, true);
+	SharedData->ModifySelectedConstraints(NewSelectedConstraints, true);
 
 }
 
@@ -4333,8 +4333,7 @@ void FPhysicsAssetEditor::HandleGraphObjectsSelected(const TArrayView<UObject*>&
 		FScopedBulkSelection BulkSelection(SharedData);
 
 		// clear selection
-		SharedData->SelectedBodies.Empty();
-		SharedData->SelectedConstraints.Empty();
+		SharedData->ClearSelected();
 
 		TArray<USkeletalBodySetup*> SelectedBodySetups;
 		TArray<UPhysicsConstraintTemplate*> SelectedConstraintTemplates;
@@ -4366,8 +4365,10 @@ void FPhysicsAssetEditor::HandleGraphObjectsSelected(const TArrayView<UObject*>&
 			}
 		}
 
-		SharedData->SetSelectedBodiesAnyPrimitive(SelectedBodyIndices, true);
-		SharedData->SetSelectedConstraints(SelectedConstraintIndices, true);
+		TArray<FPhysicsAssetEditorSharedData::FSelection> SelectedObjects;
+		SelectedObjects.Append(MakeBodySelection(SharedData->PhysicsAsset, SelectedBodyIndices));
+		SelectedObjects.Append(MakeConstraintSelection(SelectedConstraintIndices));
+		SharedData->ModifySelected(SelectedObjects, true);
 
 		SkeletonTree->SelectItemsBy([&SelectedBodySetups, &SelectedConstraintTemplates](const TSharedRef<ISkeletonTreeItem>& InItem, bool& bInOutExpand)
 		{
@@ -4422,35 +4423,26 @@ void FPhysicsAssetEditor::HandleSelectionChanged(const TArrayView<TSharedPtr<ISk
 			// Block selection broadcast until we have selected all, as this can be an expensive operation
 			FScopedBulkSelection BulkSelection(SharedData);
 
-			if (!SharedData->IsGroupSelectionActive())
-			{
-				// clear selection
-				SharedData->ClearSelectedBody();
-				SharedData->ClearSelectedConstraints();
-			}
-
 			bool bBoneSelected = false;
-			TArray<FPhysicsAssetEditorSharedData::FSelection> SelectedPrimitives;
-			TArray<int32> SelectedBodies;
-			TArray<int32> SelectedConstraints;
+			TArray<FPhysicsAssetEditorSharedData::FSelection> SelectedElements;
+
 			for (const TSharedPtr<ISkeletonTreeItem>& Item : InSelectedItems)
 			{
 				if (Item->IsOfType<FSkeletonTreePhysicsBodyItem>())
 				{
 					TSharedPtr<FSkeletonTreePhysicsBodyItem> SkeletonTreePhysicsBodyItem = StaticCastSharedPtr<FSkeletonTreePhysicsBodyItem>(Item);
-					SelectedBodies.Add(SkeletonTreePhysicsBodyItem->GetBodySetupIndex());
+					SelectedElements.Add(MakeBodySelection(SharedData->PhysicsAsset, SkeletonTreePhysicsBodyItem->GetBodySetupIndex()));
 				}
 				else if (Item->IsOfType<FSkeletonTreePhysicsShapeItem>())
 				{
 					TSharedPtr<FSkeletonTreePhysicsShapeItem> SkeletonTreePhysicsShapeItem = StaticCastSharedPtr<FSkeletonTreePhysicsShapeItem>(Item);
-					FPhysicsAssetEditorSharedData::FSelection Selection(SkeletonTreePhysicsShapeItem->GetBodySetupIndex(), SkeletonTreePhysicsShapeItem->GetShapeType(), SkeletonTreePhysicsShapeItem->GetShapeIndex());
-					SelectedPrimitives.Add(Selection);
+					SelectedElements.Add(MakePrimitiveSelection(SkeletonTreePhysicsShapeItem->GetBodySetupIndex(), SkeletonTreePhysicsShapeItem->GetShapeType(), SkeletonTreePhysicsShapeItem->GetShapeIndex()));
 
 				}
 				else if (Item->IsOfType<FSkeletonTreePhysicsConstraintItem>())
 				{
 					TSharedPtr<FSkeletonTreePhysicsConstraintItem> SkeletonTreePhysicsConstraintItem = StaticCastSharedPtr<FSkeletonTreePhysicsConstraintItem>(Item);
-					SelectedConstraints.Add(SkeletonTreePhysicsConstraintItem->GetConstraintIndex());
+					SelectedElements.Add(MakeConstraintSelection(SkeletonTreePhysicsConstraintItem->GetConstraintIndex()));
 				}
 				else if(Item->IsOfTypeByName(TEXT("FSkeletonTreeBoneItem")))
 				{
@@ -4458,9 +4450,14 @@ void FPhysicsAssetEditor::HandleSelectionChanged(const TArrayView<TSharedPtr<ISk
 				}
 			}
 
-			SharedData->SetSelectedBodies(SelectedPrimitives, true);
-			SharedData->SetSelectedBodiesFromDirectlInput(SelectedBodies, true);
-			SharedData->SetSelectedConstraints(SelectedConstraints, true);
+			if (SharedData->IsGroupSelectionActive())
+			{
+				SharedData->ModifySelected(SelectedElements, true);
+			}
+			else
+			{
+				SharedData->SetSelected(SelectedElements);
+			}
 
 			if(!bBoneSelected)
 			{
@@ -4718,7 +4715,7 @@ void FPhysicsAssetEditor::RefreshFilter()
 {
 	SkeletonTree->RefreshFilter();
 	// make sure we resynchronize the list 
-	HandleViewportSelectionChanged(SharedData->SelectedBodies, SharedData->SelectedConstraints, SharedData->SelectedCoMs);
+	HandleViewportSelectionChanged(SharedData->SelectedObjects->SelectedElements());
 }
 
 void FPhysicsAssetEditor::HandleCreateNewConstraint(int32 BodyIndex0, int32 BodyIndex1)
diff --git a/Engine/Source/Editor/PhysicsAssetEditor/Private/PhysicsAssetEditor.h b/Engine/Source/Editor/PhysicsAssetEditor/Private/PhysicsAssetEditor.h
--- a/Engine/Source/Editor/PhysicsAssetEditor/Private/PhysicsAssetEditor.h
+++ b/Engine/Source/Editor/PhysicsAssetEditor/Private/PhysicsAssetEditor.h
@@ -25,6 +25,7 @@
 
 class IDetailLayoutBuilder;
 struct FAssetData;
+struct FPhysicsAssetEditorSelectedElement;
 class FPhysicsAssetEditorTreeInfo;
 class IDetailsView;
 class SComboButton;
@@ -74,7 +75,7 @@ public:
 	TSharedPtr<FPhysicsAssetEditorSharedData> GetSharedData() const;
 
 	/** Handles a group selection change... assigns the proper object to the properties widget and the hierarchy tree view */
-	void HandleViewportSelectionChanged(const TArray<FPhysicsAssetEditorSharedData::FSelection>& InSelectedBodies, const TArray<FPhysicsAssetEditorSharedData::FSelection>& InSelectedConstraints, const TArray<FPhysicsAssetEditorSharedData::FSelection>& InSelectedCoMs);
+	void HandleViewportSelectionChanged(const TArray<FPhysicsAssetEditorSelectedElement>& InSelectedElements);
 
 	/** Repopulates the hierarchy tree view */
 	void RefreshHierachyTree();
diff --git a/Engine/Source/Editor/PhysicsAssetEditor/Private/PhysicsAssetEditorBodySetupDetailsCustomization.cpp b/Engine/Source/Editor/PhysicsAssetEditor/Private/PhysicsAssetEditorBodySetupDetailsCustomization.cpp
--- a/Engine/Source/Editor/PhysicsAssetEditor/Private/PhysicsAssetEditorBodySetupDetailsCustomization.cpp
+++ b/Engine/Source/Editor/PhysicsAssetEditor/Private/PhysicsAssetEditorBodySetupDetailsCustomization.cpp
@@ -8,6 +8,7 @@
 #include "Framework/Application/SlateApplication.h"
 #include "ObjectEditorUtils.h"
 #include "PhysicsAssetEditor.h"
+#include "PhysicsAssetEditorSelection.h"
 #include "PhysicsAssetEditorSharedData.h"
 #include "PhysicsAssetEditorSkeletalMeshComponent.h"
 #include "PhysicsEngine/BodySetup.h"
@@ -33,7 +34,7 @@ void OnGeometryPropertyPreChange(TSharedPtr<IPropertyHandle> PropertyHandle, FPh
 {
 	if (SharedData)
 	{
-		SharedData->RecordSelectedBodiesCoM();
+ 		SharedData->RecordSelectedCoM();
 	}
 }
 
@@ -41,7 +42,7 @@ void OnGeometryPropertyChanged(TSharedPtr<IPropertyHandle> PropertyHandle, FPhys
 {
 	if (SharedData)
 	{
-		SharedData->UpdateSelectedBodiesCoM();
+		SharedData->PostManipulationUpdateCoM();
 		SharedData->RefreshPhysicsAssetChange(SharedData->PhysicsAsset, false);
 	}
 }
@@ -227,7 +228,7 @@ bool FPhysicsAssetEditorBodySetupDetailsCustomization::IsCOMAxisFixedInComponent
 
 	if (SharedData)
 	{
-		for (const FPhysicsAssetEditorSharedData::FSelection& SelectedBody : SharedData->GetUniqueSelectedBodies())
+		for (const FPhysicsAssetEditorSharedData::FSelection& SelectedBody : SharedData->UniqueSelectionReferencingBodies())
 		{
 			bIsCOMFixed &= SharedData->IsCoMAxisFixedInComponentSpace(SelectedBody.Index, Axis);
 		}
@@ -240,7 +241,7 @@ void FPhysicsAssetEditorBodySetupDetailsCustomization::SetCoMAxisFixedInComponen
 {
 	if (SharedData)
 	{
-		for (const FPhysicsAssetEditorSharedData::FSelection& SelectedBody : SharedData->GetUniqueSelectedBodies())
+		for (const FPhysicsAssetEditorSharedData::FSelection& SelectedBody : SharedData->UniqueSelectionReferencingBodies())
 		{
 			SharedData->SetCoMAxisFixedInComponentSpace(SelectedBody.Index, Axis, bCOMFixed);
 		}
diff --git a/Engine/Source/Editor/PhysicsAssetEditor/Private/PhysicsAssetEditorEditMode.cpp b/Engine/Source/Editor/PhysicsAssetEditor/Private/PhysicsAssetEditorEditMode.cpp
--- a/Engine/Source/Editor/PhysicsAssetEditor/Private/PhysicsAssetEditorEditMode.cpp
+++ b/Engine/Source/Editor/PhysicsAssetEditor/Private/PhysicsAssetEditorEditMode.cpp
@@ -16,6 +16,7 @@
 #include "PhysicsAssetEditor.h"
 #include "PhysicsAssetEditorHitProxies.h"
 #include "PhysicsAssetEditorPhysicsHandleComponent.h"
+#include "PhysicsAssetEditorSelection.h"
 #include "PhysicsAssetRenderUtils.h"
 #include "DrawDebugHelpers.h"
 #include "SEditorViewport.h"
@@ -54,42 +55,42 @@ bool FPhysicsAssetEditorEditMode::GetCameraTarget(FSphere& OutTarget) const
 
 	float NumSelected = 0.0f;
 	FBox Bounds(ForceInit);
-	for (int32 BodyIndex = 0; BodyIndex < SharedData->SelectedBodies.Num(); ++BodyIndex)
+
+	for (const FPhysicsAssetEditorSharedData::FSelection& SelectedObject : SharedData->UniqueSelectionReferencingBodies())
 	{
-		FPhysicsAssetEditorSharedData::FSelection& SelectedObject = SharedData->SelectedBodies[BodyIndex];
-		int32 BoneIndex = SharedData->EditorSkelComp->GetBoneIndex(SharedData->PhysicsAsset->SkeletalBodySetups[SelectedObject.Index]->BoneName);
-		UBodySetup* BodySetup = SharedData->PhysicsAsset->SkeletalBodySetups[SelectedObject.Index];
+		int32 BoneIndex = SharedData->EditorSkelComp->GetBoneIndex(SharedData->PhysicsAsset->SkeletalBodySetups[SelectedObject.GetIndex()]->BoneName);
+		UBodySetup* BodySetup = SharedData->PhysicsAsset->SkeletalBodySetups[SelectedObject.GetIndex()];
 		const FKAggregateGeom& AggGeom = BodySetup->AggGeom;
 
 		FTransform BoneTM = SharedData->EditorSkelComp->GetBoneTransform(BoneIndex);
 		const float Scale = BoneTM.GetScale3D().GetAbsMax();
 		BoneTM.RemoveScaling();
 
-		if (SelectedObject.PrimitiveType == EAggCollisionShape::Sphere)
+		if (SelectedObject.GetPrimitiveType() == EAggCollisionShape::Sphere)
 		{
 			Bounds += AggGeom.SphereElems[SelectedObject.PrimitiveIndex].CalcAABB(BoneTM, Scale);
 		}
-		else if (SelectedObject.PrimitiveType == EAggCollisionShape::Box)
+		else if (SelectedObject.GetPrimitiveType() == EAggCollisionShape::Box)
 		{
 			Bounds += AggGeom.BoxElems[SelectedObject.PrimitiveIndex].CalcAABB(BoneTM, Scale);
 		}
-		else if (SelectedObject.PrimitiveType == EAggCollisionShape::Sphyl)
+		else if (SelectedObject.GetPrimitiveType() == EAggCollisionShape::Sphyl)
 		{
 			Bounds += AggGeom.SphylElems[SelectedObject.PrimitiveIndex].CalcAABB(BoneTM, Scale);
 		}
-		else if (SelectedObject.PrimitiveType == EAggCollisionShape::Convex)
+		else if (SelectedObject.GetPrimitiveType() == EAggCollisionShape::Convex)
 		{
 			Bounds += AggGeom.ConvexElems[SelectedObject.PrimitiveIndex].CalcAABB(BoneTM, BoneTM.GetScale3D());
 		}
-		else if (SelectedObject.PrimitiveType == EAggCollisionShape::TaperedCapsule)
+		else if (SelectedObject.GetPrimitiveType() == EAggCollisionShape::TaperedCapsule)
 		{
 			Bounds += AggGeom.TaperedCapsuleElems[SelectedObject.PrimitiveIndex].CalcAABB(BoneTM, Scale);
 		}
-		else if (SelectedObject.PrimitiveType == EAggCollisionShape::LevelSet)
+		else if (SelectedObject.GetPrimitiveType() == EAggCollisionShape::LevelSet)
 		{
 			Bounds += AggGeom.LevelSetElems[SelectedObject.PrimitiveIndex].CalcAABB(BoneTM, BoneTM.GetScale3D());
 		}
-		else if (SelectedObject.PrimitiveType == EAggCollisionShape::SkinnedLevelSet)
+		else if (SelectedObject.GetPrimitiveType() == EAggCollisionShape::SkinnedLevelSet)
 		{
 			Bounds += AggGeom.SkinnedLevelSetElems[SelectedObject.PrimitiveIndex].CalcAABB(BoneTM, BoneTM.GetScale3D());
 		}
@@ -97,11 +98,9 @@ bool FPhysicsAssetEditorEditMode::GetCameraTarget(FSphere& OutTarget) const
 		bHandled = true;
 	}
 
-	for (int32 ConstraintIndex = 0; ConstraintIndex < SharedData->SelectedConstraints.Num(); ++ConstraintIndex)
+	for (const FPhysicsAssetEditorSharedData::FSelection& SelectedConstraint : SharedData->SelectedConstraints())
 	{
-		FPhysicsAssetEditorSharedData::FSelection& SelectedObject = SharedData->SelectedConstraints[ConstraintIndex];
-		Bounds += SharedData->GetConstraintWorldTM(&SelectedObject, EConstraintFrame::Frame2).GetLocation();
-
+		Bounds += SharedData->GetConstraintWorldTM(&SelectedConstraint, EConstraintFrame::Frame2).GetLocation();
 		bHandled = true;
 	}
 
@@ -149,11 +148,13 @@ bool FPhysicsAssetEditorEditMode::EndTransform(const FGizmoState& InState)
 
 bool FPhysicsAssetEditorEditMode::HandleBeginTransform()
 {
-	if (!SharedData->bRunningSimulation && !SharedData->bManipulating)
+	if (!SharedData->bRunningSimulation && !SharedData->IsManipulating())
 	{
-		if(SharedData->GetSelectedBody() || SharedData->GetSelectedConstraint() || SharedData->GetSelectedCoM())
+		bool bManipulating = false;
+
+		if(SharedData->GetSelectedBodyOrPrimitive() || SharedData->GetSelectedConstraint() || SharedData->GetSelectedCoM())
 		{
-			if(SharedData->GetSelectedBody())
+			if(SharedData->GetSelectedBodyOrPrimitive())
 			{
 				GEditor->BeginTransaction(NSLOCTEXT("UnrealEd", "MoveElement", "Move Element"));
 			}
@@ -167,87 +168,81 @@ bool FPhysicsAssetEditorEditMode::HandleBeginTransform()
 			}
 		}
 
-		if (SharedData->GetSelectedBody())
+		if (SharedData->GetSelectedBodyOrPrimitive())
 		{
-			for (int32 i = 0; i<SharedData->SelectedBodies.Num(); ++i)
+			for (const FPhysicsAssetEditorSharedData::FSelection& SelectedBody : SharedData->SelectedPrimitives())
 			{
-				SharedData->PhysicsAsset->SkeletalBodySetups[SharedData->SelectedBodies[i].Index]->Modify();
-				SharedData->SelectedBodies[i].ManipulateTM = FTransform::Identity;
+				SharedData->PhysicsAsset->SkeletalBodySetups[SelectedBody.Index]->Modify();
 			}
 
-			SharedData->RecordSelectedBodiesCoM();
-			SharedData->bManipulating = true;
+			bManipulating = true;
 		}
 	
 		if (SharedData->GetSelectedCoM())
 		{
-			for (int32 i = 0; i < SharedData->SelectedCoMs.Num(); ++i)
+			for(const FPhysicsAssetEditorSharedData::FSelection& SelectedCoM : SharedData->SelectedCoMs())
 			{
-				const int32 BodyIndex = SharedData->SelectedCoMs[i].Index;
-				SharedData->PhysicsAsset->SkeletalBodySetups[BodyIndex]->Modify();
-				SharedData->SelectedCoMs[i].ManipulateTM = FTransform::Identity;
-				SharedData->SelectedCoMs[i].CoMPosition = SharedData->EditorSkelComp->Bodies[BodyIndex]->GetCOMPosition();
+				SharedData->PhysicsAsset->SkeletalBodySetups[SelectedCoM.Index]->Modify();
 			}
 
-			SharedData->bManipulating = true;
+			bManipulating = true;
 		}
 
 		if (SharedData->GetSelectedConstraint())
 		{
-			const int32 Count = SharedData->SelectedConstraints.Num();
+			const TArray<FPhysicsAssetEditorSharedData::FSelection> SelectedConstraints = SharedData->SelectedConstraints().ToArray();
+
+			const int32 Count = SelectedConstraints.Num();
+			ManConTM.SetNumUninitialized(Count);
 			StartManRelConTM.SetNumUninitialized(Count);
 			StartManParentConTM.SetNumUninitialized(Count);
 			StartManChildConTM.SetNumUninitialized(Count);
 
-			for (int32 i = 0; i < SharedData->SelectedConstraints.Num(); ++i)
+			for (int32 i = 0; i < Count; ++i)
 			{
-				FPhysicsAssetEditorSharedData::FSelection& Constraint = SharedData->SelectedConstraints[i];
+				const FPhysicsAssetEditorSharedData::FSelection& Constraint = SelectedConstraints[i];
 				SharedData->PhysicsAsset->ConstraintSetup[Constraint.Index]->Modify();
-				Constraint.ManipulateTM = FTransform::Identity;
 
 				const FTransform WParentFrame = SharedData->GetConstraintWorldTM(&Constraint, EConstraintFrame::Frame2);
 				const FTransform WChildFrame = SharedData->GetConstraintWorldTM(&Constraint, EConstraintFrame::Frame1);
 				const UPhysicsConstraintTemplate* Setup = SharedData->PhysicsAsset->ConstraintSetup[Constraint.Index];
 
+				ManConTM[i] = FTransform::Identity;
 				StartManRelConTM[i] = WChildFrame * WParentFrame.Inverse();
 				StartManParentConTM[i] = Setup->DefaultInstance.GetRefFrame(EConstraintFrame::Frame2);
 				StartManChildConTM[i] = Setup->DefaultInstance.GetRefFrame(EConstraintFrame::Frame1);
 			}
 
-			SharedData->bManipulating = true;
+			bManipulating = true;
+		}
+
+		if (bManipulating)
+		{
+			SharedData->BeginManipulation();
 		}
 	}
 	
-	return SharedData->bManipulating;
+	return SharedData->IsManipulating();
 }
 
 bool FPhysicsAssetEditorEditMode::HandleEndTransform(FEditorViewportClient* InViewportClient) const
 {
-	if (SharedData->bManipulating)
+	if (SharedData->IsManipulating())
 	{
-		SharedData->bManipulating = false;
-
-		for (int32 i = 0; i<SharedData->SelectedBodies.Num(); ++i)
+		for (const FPhysicsAssetEditorSharedData::FSelection& SelectedObject : SharedData->SelectedPrimitives())
 		{
-			FPhysicsAssetEditorSharedData::FSelection & SelectedObject = SharedData->SelectedBodies[i];
-			UBodySetup* BodySetup = SharedData->PhysicsAsset->SkeletalBodySetups[SelectedObject.Index];
+			UBodySetup* const BodySetup = SharedData->PhysicsAsset->SkeletalBodySetups[SelectedObject.GetIndex()];
 
-			FKAggregateGeom* AggGeom = &BodySetup->AggGeom;
-
-			if (SelectedObject.PrimitiveType == EAggCollisionShape::Convex)
+			if (SelectedObject.GetPrimitiveType() == EAggCollisionShape::Convex)
 			{
-				FKConvexElem& Convex = AggGeom->ConvexElems[SelectedObject.PrimitiveIndex];
-				Convex.SetTransform(SelectedObject.ManipulateTM * Convex.GetTransform());
-
 				BodySetup->InvalidatePhysicsData();
 				BodySetup->CreatePhysicsMeshes();
 			}
 		}
 
-		SharedData->bShouldUpdatedSelectedCoMs = true;
+		SharedData->EndManipulation();
 
 		GEditor->EndTransaction();
-		SharedData->RefreshPhysicsAssetChange(SharedData->PhysicsAsset, false);
 
 		if (InViewportClient)
 		{
@@ -319,7 +314,7 @@ bool FPhysicsAssetEditorEditMode::InputAxis(FEditorViewportClient* InViewportCli
 {
 	bool bHandled = false;
 	// If we are 'manipulating' don't move the camera but do something else with mouse input.
-	if (SharedData->bManipulating)
+	if (SharedData->IsManipulating())
 	{
 		bool bCtrlDown = InViewport->KeyState(EKeys::LeftControl) || InViewport->KeyState(EKeys::RightControl);
 
@@ -390,82 +385,81 @@ bool FPhysicsAssetEditorEditMode::InputDelta(FEditorViewportClient* InViewportCl
 
 	bool bHandled = false;
 	const EAxisList::Type CurrentAxis = InViewportClient->GetCurrentWidgetAxis();
-	if (!SharedData->bRunningSimulation && SharedData->bManipulating && CurrentAxis != EAxisList::None)
+	if (!SharedData->bRunningSimulation && SharedData->IsManipulating() && CurrentAxis != EAxisList::None)
 	{
-		for (int32 i = 0; i< SharedData->SelectedBodies.Num(); ++i)
+		for(const FPhysicsAssetEditorSharedData::FSelection& SelectedObject : SharedData->SelectedPrimitives())
 		{
-			FPhysicsAssetEditorSharedData::FSelection & SelectedObject = SharedData->SelectedBodies[i];
-			if (SharedData->bManipulating)
+			if (SharedData->IsManipulating())
 			{
 				float BoneScale = 1.f;
-
-				int32 BoneIndex = SharedData->EditorSkelComp->GetBoneIndex(SharedData->PhysicsAsset->SkeletalBodySetups[SelectedObject.Index]->BoneName);
+				int32 BoneIndex = SharedData->EditorSkelComp->GetBoneIndex(SharedData->PhysicsAsset->SkeletalBodySetups[SelectedObject.GetIndex()]->BoneName);
 
 				FTransform BoneTM = SharedData->EditorSkelComp->GetBoneTransform(BoneIndex);
 				BoneScale = BoneTM.GetScale3D().GetAbsMax();
 				BoneTM.RemoveScaling();
 
-				SelectedObject.WidgetTM = SharedData->EditorSkelComp->GetPrimitiveTransform(BoneTM, SelectedObject.Index, SelectedObject.PrimitiveType, SelectedObject.PrimitiveIndex, BoneScale);
-
+				FTransform SelectedObjectWidgetTM = SharedData->EditorSkelComp->GetPrimitiveTransform(BoneTM, SelectedObject.GetIndex(), SelectedObject.GetPrimitiveType(), SelectedObject.PrimitiveIndex, BoneScale);
+				FTransform ManipulateTM;
+				
 				if(InViewportClient->GetWidgetMode() == UE::Widget::WM_Translate || InViewportClient->GetWidgetMode() == UE::Widget::WM_Rotate)
 				{
 					if (InViewportClient->GetWidgetMode() == UE::Widget::WM_Translate)
 					{
-						FVector DragToUse = GetLocalTranslation(InViewportClient, InDrag, SelectedObject.WidgetTM);
-						FVector Dir = SelectedObject.WidgetTM.InverseTransformVector(DragToUse.GetSafeNormal());
+						FVector DragToUse = GetLocalTranslation(InViewportClient, InDrag, SelectedObjectWidgetTM);
+						FVector Dir = SelectedObjectWidgetTM.InverseTransformVector(DragToUse.GetSafeNormal());
 						FVector DragVec = Dir * DragToUse.Size() / BoneScale;
-						SelectedObject.ManipulateTM.AddToTranslation(DragVec);
+
+						ManipulateTM.AddToTranslation(DragVec);
 					}
 					else if (InViewportClient->GetWidgetMode() == UE::Widget::WM_Rotate)
 					{
-						FRotator RotatorToUse = GetLocalRotation(InViewportClient, InRot, SelectedObject.WidgetTM);
+						FRotator RotatorToUse = GetLocalRotation(InViewportClient, InRot, SelectedObjectWidgetTM);
 
 						FVector Axis;
 						float Angle;
 						RotatorToUse.Quaternion().ToAxisAndAngle(Axis, Angle);
 
-						Axis = SelectedObject.WidgetTM.InverseTransformVectorNoScale(Axis);
+						Axis = SelectedObjectWidgetTM.InverseTransformVectorNoScale(Axis);
 
-						const FQuat Start = SelectedObject.ManipulateTM.GetRotation();
+						const FQuat Start = ManipulateTM.GetRotation(); // TODO - get rid of this Start quat
 						const FQuat Delta = FQuat(Axis, Angle);
 						const FQuat Result = Delta * Start;
 
-						SelectedObject.ManipulateTM = FTransform(Result);
+						ManipulateTM = FTransform(Result);
 					}
 				
-					UBodySetup* BodySetup = SharedData->PhysicsAsset->SkeletalBodySetups[SelectedObject.Index];
+					UBodySetup* BodySetup = SharedData->PhysicsAsset->SkeletalBodySetups[SelectedObject.GetIndex()];
 					FKAggregateGeom* AggGeom = &BodySetup->AggGeom;
 
 					// for all but convex shapes, we apply straight away
-					if (SelectedObject.PrimitiveType == EAggCollisionShape::Sphere)
+					if (SelectedObject.GetPrimitiveType() == EAggCollisionShape::Sphere)
+					{
+						AggGeom->SphereElems[SelectedObject.PrimitiveIndex].Center = (ManipulateTM * AggGeom->SphereElems[SelectedObject.PrimitiveIndex].GetTransform()).GetLocation();
+					}
+					else if (SelectedObject.GetPrimitiveType() == EAggCollisionShape::Box)
 					{
-						AggGeom->SphereElems[SelectedObject.PrimitiveIndex].Center = (SelectedObject.ManipulateTM * AggGeom->SphereElems[SelectedObject.PrimitiveIndex].GetTransform()).GetLocation();
-						SelectedObject.ManipulateTM.SetIdentity();
+						AggGeom->BoxElems[SelectedObject.PrimitiveIndex].SetTransform(ManipulateTM * AggGeom->BoxElems[SelectedObject.PrimitiveIndex].GetTransform());
 					}
-					else if (SelectedObject.PrimitiveType == EAggCollisionShape::Box)
+					else if (SelectedObject.GetPrimitiveType() == EAggCollisionShape::Sphyl)
 					{
-						AggGeom->BoxElems[SelectedObject.PrimitiveIndex].SetTransform(SelectedObject.ManipulateTM * AggGeom->BoxElems[SelectedObject.PrimitiveIndex].GetTransform());
-						SelectedObject.ManipulateTM.SetIdentity();
+						AggGeom->SphylElems[SelectedObject.PrimitiveIndex].SetTransform(ManipulateTM * AggGeom->SphylElems[SelectedObject.PrimitiveIndex].GetTransform());
 					}
-					else if (SelectedObject.PrimitiveType == EAggCollisionShape::Sphyl)
+					else if (SelectedObject.GetPrimitiveType() == EAggCollisionShape::TaperedCapsule)
 					{
-						AggGeom->SphylElems[SelectedObject.PrimitiveIndex].SetTransform(SelectedObject.ManipulateTM * AggGeom->SphylElems[SelectedObject.PrimitiveIndex].GetTransform());
-						SelectedObject.ManipulateTM.SetIdentity();
+						AggGeom->TaperedCapsuleElems[SelectedObject.PrimitiveIndex].SetTransform(ManipulateTM * AggGeom->TaperedCapsuleElems[SelectedObject.PrimitiveIndex].GetTransform());
 					}
-					else if (SelectedObject.PrimitiveType == EAggCollisionShape::TaperedCapsule)
+					else if (SelectedObject.GetPrimitiveType() == EAggCollisionShape::LevelSet)
 					{
-						AggGeom->TaperedCapsuleElems[SelectedObject.PrimitiveIndex].SetTransform(SelectedObject.ManipulateTM * AggGeom->TaperedCapsuleElems[SelectedObject.PrimitiveIndex].GetTransform());
-						SelectedObject.ManipulateTM.SetIdentity();
+						AggGeom->LevelSetElems[SelectedObject.PrimitiveIndex].SetTransform(ManipulateTM * AggGeom->LevelSetElems[SelectedObject.PrimitiveIndex].GetTransform());
 					}
-					else if (SelectedObject.PrimitiveType == EAggCollisionShape::LevelSet)
+					else if (SelectedObject.GetPrimitiveType() == EAggCollisionShape::Convex)
 					{
-						AggGeom->LevelSetElems[SelectedObject.PrimitiveIndex].SetTransform(SelectedObject.ManipulateTM * AggGeom->LevelSetElems[SelectedObject.PrimitiveIndex].GetTransform());
-						SelectedObject.ManipulateTM.SetIdentity();
+						AggGeom->ConvexElems[SelectedObject.PrimitiveIndex].SetTransform(ManipulateTM * AggGeom->ConvexElems[SelectedObject.PrimitiveIndex].GetTransform());
 					}
 				}
 				else if (InViewportClient->GetWidgetMode() == UE::Widget::WM_Scale)
 				{
-					ModifyPrimitiveSize(SelectedObject.Index, SelectedObject.PrimitiveType, SelectedObject.PrimitiveIndex, InScale);
+					ModifyPrimitiveSize(SelectedObject.GetIndex(), SelectedObject.GetPrimitiveType(), SelectedObject.GetPrimitiveIndex(), InScale);
 				}
 
 				bHandled = true;
@@ -477,60 +471,63 @@ bool FPhysicsAssetEditorEditMode::InputDelta(FEditorViewportClient* InViewportCl
 			SharedData->UpdateClothPhysics();
 		}
 
-		for (int32 i = 0; i<SharedData->SelectedConstraints.Num(); ++i)
+		for (UPhysicsAssetEditorSelection::Iterator SelectedConstraintIterator = SharedData->SelectedConstraints().CreateConstIterator(); SelectedConstraintIterator; ++SelectedConstraintIterator)
 		{
-			FPhysicsAssetEditorSharedData::FSelection & SelectedObject = SharedData->SelectedConstraints[i];
-			if (SharedData->bManipulating)
+			if (SharedData->IsManipulating())
 			{
+				const FPhysicsAssetEditorSharedData::FSelection& SelectedObject = *SelectedConstraintIterator;				
+				const int32 SelectedConstraintIndex = SelectedConstraintIterator.GetIndexIntoFilteredSelection();
+
 				float BoneScale = 1.f;
-				SelectedObject.WidgetTM = SharedData->GetConstraintMatrix(SelectedObject.Index, GetConstraintFrameForWidget(), 1.f);
+				FTransform SelectedObjectWidgetTM = SharedData->GetConstraintMatrix(SelectedObject.GetIndex(), GetConstraintFrameForWidget(), 1.f);
+				FTransform& ManipulateTM = ManConTM[SelectedConstraintIndex];
 
 				if (InViewportClient->GetWidgetMode() == UE::Widget::WM_Translate)
 				{
-					FVector DragToUse = GetLocalTranslation(InViewportClient, InDrag, SelectedObject.WidgetTM);
-					FVector Dir = SelectedObject.WidgetTM.InverseTransformVector(DragToUse.GetSafeNormal());
+					FVector DragToUse = GetLocalTranslation(InViewportClient, InDrag, SelectedObjectWidgetTM);
+					FVector Dir = SelectedObjectWidgetTM.InverseTransformVector(DragToUse.GetSafeNormal());
 					FVector DragVec = Dir * DragToUse.Size() / BoneScale;
-					SelectedObject.ManipulateTM.AddToTranslation(DragVec);
+					ManipulateTM.AddToTranslation(DragVec);
 				}
 				else if (InViewportClient->GetWidgetMode() == UE::Widget::WM_Rotate)
 				{
-					FRotator RotatorToUse = GetLocalRotation(InViewportClient, InRot, SelectedObject.WidgetTM);
+					FRotator RotatorToUse = GetLocalRotation(InViewportClient, InRot, SelectedObjectWidgetTM);
 
 					FVector Axis;
 					float Angle;
 					RotatorToUse.Quaternion().ToAxisAndAngle(Axis, Angle);
 
-					Axis = SelectedObject.WidgetTM.InverseTransformVectorNoScale(Axis);
+					Axis = SelectedObjectWidgetTM.InverseTransformVectorNoScale(Axis);
 
-					const FQuat Start = SelectedObject.ManipulateTM.GetRotation();
+					const FQuat Start = ManipulateTM.GetRotation();
 					const FQuat Delta = FQuat(Axis, Angle);
 					const FQuat Result = Delta * Start;
 
-					SelectedObject.ManipulateTM = FTransform(Result);
+					ManipulateTM = FTransform(Result);
 				}
 
 				// Apply manipulations to Child or Parent or both transforms according to the constraint's view port manipulation flags.
 				{
-					UPhysicsConstraintTemplate* ConstraintSetup = SharedData->PhysicsAsset->ConstraintSetup[SelectedObject.Index];
+					UPhysicsConstraintTemplate* ConstraintSetup = SharedData->PhysicsAsset->ConstraintSetup[SelectedObject.GetIndex()];
 					FPhysicsAssetRenderSettings* const RenderSettings = SharedData->GetRenderSettings();
 
 					if (RenderSettings && !EnumHasAnyFlags(RenderSettings->ConstraintViewportManipulationFlags, EConstraintTransformComponentFlags::AllChild))
 					{
 						// Rotate or move the parent transform only.
-						ConstraintSetup->DefaultInstance.SetRefFrame(EConstraintFrame::Frame2, SelectedObject.ManipulateTM * StartManParentConTM[i]);
-						ConstraintSetup->DefaultInstance.SetRefFrame(EConstraintFrame::Frame1, FTransform(StartManChildConTM[i]));
+						ConstraintSetup->DefaultInstance.SetRefFrame(EConstraintFrame::Frame2, ManipulateTM * StartManParentConTM[SelectedConstraintIndex]);
+						ConstraintSetup->DefaultInstance.SetRefFrame(EConstraintFrame::Frame1, FTransform(StartManChildConTM[SelectedConstraintIndex]));
 					}
 					else if (RenderSettings && !EnumHasAnyFlags(RenderSettings->ConstraintViewportManipulationFlags, EConstraintTransformComponentFlags::AllParent))
 					{
 						// Rotate or move the child transform only.
-						ConstraintSetup->DefaultInstance.SetRefFrame(EConstraintFrame::Frame1, SelectedObject.ManipulateTM * StartManChildConTM[i]);
-						ConstraintSetup->DefaultInstance.SetRefFrame(EConstraintFrame::Frame2, FTransform(StartManParentConTM[i]));
+						ConstraintSetup->DefaultInstance.SetRefFrame(EConstraintFrame::Frame1, ManipulateTM * StartManChildConTM[SelectedConstraintIndex]);
+						ConstraintSetup->DefaultInstance.SetRefFrame(EConstraintFrame::Frame2, FTransform(StartManParentConTM[SelectedConstraintIndex]));
 					}
 					else
 					{
 						// Rotate or move both the parent and child transform.
-						ConstraintSetup->DefaultInstance.SetRefFrame(EConstraintFrame::Frame2, SelectedObject.ManipulateTM * StartManParentConTM[i]);
-						SharedData->SetConstraintRelTM(&SelectedObject, StartManRelConTM[i]);
+						ConstraintSetup->DefaultInstance.SetRefFrame(EConstraintFrame::Frame2, ManipulateTM * StartManParentConTM[SelectedConstraintIndex]);
+						SharedData->SetConstraintRelTM(&SelectedObject, StartManRelConTM[SelectedConstraintIndex]);
 					}
 
 					bHandled = true;
@@ -538,32 +535,39 @@ bool FPhysicsAssetEditorEditMode::InputDelta(FEditorViewportClient* InViewportCl
 			}
 		}
 
-		for (int32 i = 0; i < SharedData->SelectedCoMs.Num(); ++i)
-		{
-			FPhysicsAssetEditorSharedData::FSelection& SelectedObject = SharedData->SelectedCoMs[i];
-			if (SharedData->bManipulating)
+		for(UPhysicsAssetEditorSelection::Iterator SelectedCoMIterator = SharedData->SelectedCoMs().CreateConstIterator(); SelectedCoMIterator; ++SelectedCoMIterator)
+		{	
+			if (SharedData->IsManipulating())
 			{
-				const int32 BodyIndex = SelectedObject.Index;
-				const int32 BoneIndex = SharedData->EditorSkelComp->GetBoneIndex(SharedData->PhysicsAsset->SkeletalBodySetups[BodyIndex]->BoneName);
-				FTransform BoneTM = SharedData->EditorSkelComp->GetBoneTransform(BoneIndex);
-				const float BoneScale = BoneTM.GetScale3D().GetAbsMax();
-				BoneTM.RemoveScaling();
-
-				SelectedObject.WidgetTM.SetRotation(BoneTM.GetRotation());
-				SelectedObject.WidgetTM.SetTranslation(SelectedObject.CoMPosition);
-	
-				if (InViewportClient->GetWidgetMode() == UE::Widget::WM_Translate)
+				const FPhysicsAssetEditorSharedData::FSelection& SelectedObject = *SelectedCoMIterator;
+				FVector* const ManipulatedBodyCoMPosition = SharedData->FindManipulatedBodyCoMPosition(SelectedObject.GetIndex());
+				check(ManipulatedBodyCoMPosition); // Expect to find this for any CoM undergoing manipulation.
+				if (ManipulatedBodyCoMPosition)
 				{
-					FVector DragToUse = GetLocalTranslation(InViewportClient, InDrag, SelectedObject.WidgetTM);
-					FVector Dir = DragToUse.GetSafeNormal();
-					FVector DragVec = Dir * DragToUse.Size() / BoneScale;
+					const int32 SelectedCoMIndex = SelectedCoMIterator.GetIndexIntoFilteredSelection();
 
-					SelectedObject.CoMPosition += DragVec; // Use widget movement to update selected objects target worldspace location, then calculate the correct CoM nudge to position the world space CoM at this location when manipulation ends.
+					const int32 BodyIndex = SelectedObject.GetIndex();
+					const int32 BoneIndex = SharedData->EditorSkelComp->GetBoneIndex(SharedData->PhysicsAsset->SkeletalBodySetups[BodyIndex]->BoneName);
+					FTransform BoneTM = SharedData->EditorSkelComp->GetBoneTransform(BoneIndex);
+					const float BoneScale = BoneTM.GetScale3D().GetAbsMax();
+					BoneTM.RemoveScaling();
+					FTransform SelectedObjectWidgetTM;
+					SelectedObjectWidgetTM.SetRotation(BoneTM.GetRotation());
+					SelectedObjectWidgetTM.SetTranslation(*ManipulatedBodyCoMPosition);
 
-					// Update CoM nudge value in body instance to ensure the values in the details panel update as the manipulator is moved in the viewport.
-					const FVector CalculatedCoMOffset = SharedData->CalculateCoMNudgeForWorldSpacePosition(BodyIndex, SelectedObject.CoMPosition);
-					SharedData->EditorSkelComp->Bodies[BodyIndex]->COMNudge = CalculatedCoMOffset;
-					SharedData->PhysicsAsset->SkeletalBodySetups[BodyIndex]->DefaultInstance.COMNudge = CalculatedCoMOffset;
+					if (InViewportClient->GetWidgetMode() == UE::Widget::WM_Translate)
+					{
+						FVector DragToUse = GetLocalTranslation(InViewportClient, InDrag, SelectedObjectWidgetTM);
+						FVector Dir = DragToUse.GetSafeNormal();
+						FVector DragVec = Dir * DragToUse.Size() / BoneScale;
+
+						*ManipulatedBodyCoMPosition += DragVec; // Use widget movement to update selected objects target worldspace location, then calculate the correct CoM nudge to position the world space CoM at this location when manipulation ends.
+
+						// Update CoM nudge value in body instance to ensure the values in the details panel update as the manipulator is moved in the viewport.
+						const FVector CalculatedCoMOffset = SharedData->CalculateCoMNudgeForWorldSpacePosition(BodyIndex, *ManipulatedBodyCoMPosition);
+						SharedData->EditorSkelComp->Bodies[BodyIndex]->COMNudge = CalculatedCoMOffset;
+						SharedData->PhysicsAsset->SkeletalBodySetups[BodyIndex]->DefaultInstance.COMNudge = CalculatedCoMOffset;
+					}
 				}
 			}
 
@@ -629,7 +633,7 @@ void FPhysicsAssetEditorEditMode::Tick(FEditorViewportClient* ViewportClient, fl
 		SharedData->SetGroupSelectionActive(ViewportClient->IsCtrlPressed() || ViewportClient->IsShiftPressed());
 	}
 
-	UpdateCoM();
+	SharedData->UpdateCoM();
 }
 
 void FPhysicsAssetEditorEditMode::Render(const FSceneView* View, FViewport* Viewport, FPrimitiveDrawInterface* PDI)
@@ -746,7 +750,7 @@ bool FPhysicsAssetEditorEditMode::AllowWidgetMove()
 
 bool FPhysicsAssetEditorEditMode::ShouldDrawWidget() const
 {
-	return !SharedData->bRunningSimulation && (SharedData->GetSelectedBody() || SharedData->GetSelectedConstraint() || SharedData->GetSelectedCoM());
+	return !SharedData->bRunningSimulation && (SharedData->GetSelectedBodyOrPrimitive() || SharedData->GetSelectedConstraint() || SharedData->GetSelectedCoM());
 }
 
 bool FPhysicsAssetEditorEditMode::UsesTransformWidget() const
@@ -787,33 +791,43 @@ EConstraintFrame::Type FPhysicsAssetEditorEditMode::GetConstraintFrameForWidget(
 bool FPhysicsAssetEditorEditMode::GetCustomDrawingCoordinateSystem(FMatrix& InMatrix, void* InData)
 {
 	// Don't draw widget if nothing selected.
-	if (SharedData->GetSelectedBody())
+	if (const FPhysicsAssetEditorSharedData::FSelection* const SelectedPrimitive = SharedData->GetSelectedPrimitive())
 	{
-		int32 BoneIndex = SharedData->EditorSkelComp->GetBoneIndex(SharedData->PhysicsAsset->SkeletalBodySetups[SharedData->GetSelectedBody()->Index]->BoneName);
-		if(BoneIndex != INDEX_NONE)
+		if (SharedData->PhysicsAsset->SkeletalBodySetups.IsValidIndex(SelectedPrimitive->GetIndex())) // TODO - investigate - why do we sometimes get an invalid selected object when this fn is called via ForEachEdMode from FEditorModeTools::GetCustomCoordinateSystem ?
 		{
-			FTransform BoneTM = SharedData->EditorSkelComp->GetBoneTransform(BoneIndex);
-			BoneTM.RemoveScaling();
+			const FName BoneName = SharedData->PhysicsAsset->SkeletalBodySetups[SelectedPrimitive->GetIndex()]->BoneName;
+			const int32 BoneIndex = SharedData->EditorSkelComp->GetBoneIndex(BoneName);
+
+			if (BoneIndex != INDEX_NONE)
+			{
+				FTransform BoneTM = SharedData->EditorSkelComp->GetBoneTransform(BoneIndex);
+				BoneTM.RemoveScaling();
 
-			InMatrix = SharedData->EditorSkelComp->GetPrimitiveTransform(BoneTM, SharedData->GetSelectedBody()->Index, SharedData->GetSelectedBody()->PrimitiveType, SharedData->GetSelectedBody()->PrimitiveIndex, 1.f).ToMatrixNoScale().RemoveTranslation();
-			return true;
+				InMatrix = SharedData->EditorSkelComp->GetPrimitiveTransform(BoneTM, SelectedPrimitive->GetIndex(), SelectedPrimitive->GetPrimitiveType(), SelectedPrimitive->GetPrimitiveIndex(), 1.f).ToMatrixNoScale().RemoveTranslation();
+				return true;
+			}
 		}
 	}
 	else if (SharedData->GetSelectedCoM())
 	{
-		const int32 BoneIndex = SharedData->EditorSkelComp->GetBoneIndex(SharedData->PhysicsAsset->SkeletalBodySetups[SharedData->GetSelectedCoM()->Index]->BoneName);
-		if (BoneIndex != INDEX_NONE)
+		const int32 BodyIndex = SharedData->GetSelectedCoM()->GetIndex();
+
+		if (SharedData->PhysicsAsset->SkeletalBodySetups.IsValidIndex(BodyIndex))
 		{
-			FTransform SelectedCoMTM = SharedData->EditorSkelComp->GetBoneTransform(BoneIndex);
-			SelectedCoMTM.RemoveScaling();
+			const int32 BoneIndex = SharedData->EditorSkelComp->GetBoneIndex(SharedData->PhysicsAsset->SkeletalBodySetups[BodyIndex]->BoneName);
+			if (BoneIndex != INDEX_NONE)
+			{
+				FTransform SelectedCoMTM = SharedData->EditorSkelComp->GetBoneTransform(BoneIndex);
+				SelectedCoMTM.RemoveScaling();
 
-			InMatrix = SelectedCoMTM.ToMatrixNoScale().RemoveTranslation();
-			return true;
+				InMatrix = SelectedCoMTM.ToMatrixNoScale().RemoveTranslation();
+				return true;
+			}
 		}
 	}
 	else if (SharedData->GetSelectedConstraint())
 	{
-		InMatrix = SharedData->GetConstraintMatrix(SharedData->GetSelectedConstraint()->Index, GetConstraintFrameForWidget(), 1.f).ToMatrixNoScale().RemoveTranslation();
+		InMatrix = SharedData->GetConstraintMatrix(SharedData->GetSelectedConstraint()->GetIndex(), GetConstraintFrameForWidget(), 1.f).ToMatrixNoScale().RemoveTranslation();
 		return true;
 	}
 
@@ -828,33 +842,25 @@ bool FPhysicsAssetEditorEditMode::GetCustomInputCoordinateSystem(FMatrix& InMatr
 FVector FPhysicsAssetEditorEditMode::GetWidgetLocation() const
 {
 	// Don't draw widget if nothing selected.
-	if (SharedData->GetSelectedBody())
+	if (SharedData->GetSelectedBodyOrPrimitive())
 	{
-		int32 BoneIndex = SharedData->EditorSkelComp->GetBoneIndex(SharedData->PhysicsAsset->SkeletalBodySetups[SharedData->GetSelectedBody()->Index]->BoneName);
+		int32 BoneIndex = SharedData->EditorSkelComp->GetBoneIndex(SharedData->PhysicsAsset->SkeletalBodySetups[SharedData->GetSelectedBodyOrPrimitive()->Index]->BoneName);
 		if(BoneIndex != INDEX_NONE)
 		{
 			FTransform BoneTM = SharedData->EditorSkelComp->GetBoneTransform(BoneIndex);
 			const float Scale = BoneTM.GetScale3D().GetAbsMax();
 			BoneTM.RemoveScaling();
 
-			return SharedData->EditorSkelComp->GetPrimitiveTransform(BoneTM, SharedData->GetSelectedBody()->Index, SharedData->GetSelectedBody()->PrimitiveType, SharedData->GetSelectedBody()->PrimitiveIndex, Scale).GetTranslation();
+			return SharedData->EditorSkelComp->GetPrimitiveTransform(BoneTM, SharedData->GetSelectedBodyOrPrimitive()->GetIndex(), SharedData->GetSelectedBodyOrPrimitive()->GetPrimitiveType(), SharedData->GetSelectedBodyOrPrimitive()->GetPrimitiveIndex(), Scale).GetTranslation();
 		}
 	}
 	else if (SharedData->GetSelectedCoM())
 	{
-		if (SharedData->bManipulating)
-		{
-			// return the CoM position from the FSelection object because the physics body's CoM position will only be updated at the end of manipulation.
-			return SharedData->GetSelectedCoM()->CoMPosition;
-		}
-		else
-		{
-			return SharedData->EditorSkelComp->Bodies[SharedData->GetSelectedCoM()->Index]->GetCOMPosition();
-		}
+		return SharedData->GetSelectedCoMPosition();
 	}
 	else if (SharedData->GetSelectedConstraint())
 	{
-		return SharedData->GetConstraintMatrix(SharedData->GetSelectedConstraint()->Index, GetConstraintFrameForWidget(), 1.f).GetTranslation();
+		return SharedData->GetConstraintMatrix(SharedData->GetSelectedConstraint()->GetIndex(), GetConstraintFrameForWidget(), 1.f).GetTranslation();
 	}
 
 	return FVector::ZeroVector;
@@ -901,14 +907,15 @@ bool FPhysicsAssetEditorEditMode::HandleClick(FEditorViewportClient* InViewportC
 				HPhysicsAssetEditorEdBoneProxy* BoneProxy = (HPhysicsAssetEditorEdBoneProxy*)HitProxy;
 
 				// Select body under cursor if not already selected.
-				FPhysicsAssetEditorSharedData::FSelection Selection(BoneProxy->BodyIndex, BoneProxy->PrimType, BoneProxy->PrimIndex);
-				if (!SharedData->IsBodySelected(Selection))
+				FPhysicsAssetEditorSharedData::FSelection Selection = MakePrimitiveSelection(BoneProxy->BodyIndex, BoneProxy->PrimType, BoneProxy->PrimIndex);
+				
+				if (!SharedData->IsSelected(Selection))
 				{
-					SharedData->SetSelectedBody(Selection, true);
+					SharedData->ModifySelectedBodies(Selection, true);
 				}
-
+				
 				// Pop up menu, if we have a body selected.
-				if (SharedData->GetSelectedBody())
+				if (SharedData->GetSelectedBodyOrPrimitive())
 				{
 					OpenBodyMenu(InViewportClient);
 				}
@@ -922,7 +929,7 @@ bool FPhysicsAssetEditorEditMode::HandleClick(FEditorViewportClient* InViewportC
 				// Select constraint under cursor if not already selected
 				if (!SharedData->IsConstraintSelected(ConstraintProxy->ConstraintIndex))
 				{
-					SharedData->SetSelectedConstraint(ConstraintProxy->ConstraintIndex, true);
+					SharedData->ModifySelectedConstraints(ConstraintProxy->ConstraintIndex, true);
 				}
 
 				// Pop up menu, if we have a constraint selected.
@@ -1170,7 +1177,7 @@ bool FPhysicsAssetEditorEditMode::SimMouseWheelDown(FEditorViewportClient* InVie
 
 void FPhysicsAssetEditorEditMode::ModifyPrimitiveSize(int32 BodyIndex, EAggCollisionShape::Type PrimType, int32 PrimIndex, FVector DeltaSize)
 {
-	check(SharedData->GetSelectedBody());
+	check(SharedData->GetSelectedBodyOrPrimitive());
 
 	FKAggregateGeom* AggGeom = &SharedData->PhysicsAsset->SkeletalBodySetups[BodyIndex]->AggGeom;
 
@@ -1230,23 +1237,4 @@ void FPhysicsAssetEditorEditMode::HitNothing(FEditorViewportClient* InViewportCl
 	PhysicsAssetEditorPtr.Pin()->RefreshHierachyTree();
 }
 
-void FPhysicsAssetEditorEditMode::UpdateCoM()
-{
-	if (SharedData->bShouldUpdatedSelectedCoMs) // < This calculation must be delayed by a frame s.t. changes to the physics state have been propagated to the physics bodies.
-	{
-		SharedData->UpdateSelectedBodiesCoM();
-		
-		for (int32 i = 0; i < SharedData->SelectedCoMs.Num(); ++i)
-		{
-			FPhysicsAssetEditorSharedData::FSelection& SelectedObject = SharedData->SelectedCoMs[i];
-			const FVector CalculatedCoMOffset = SharedData->CalculateCoMNudgeForWorldSpacePosition(SelectedObject.Index, SelectedObject.CoMPosition);
-			SharedData->PhysicsAsset->SkeletalBodySetups[SelectedObject.Index]->DefaultInstance.COMNudge = CalculatedCoMOffset;
-		}
-
-		SharedData->RefreshPhysicsAssetChange(SharedData->PhysicsAsset, false);
-
-		SharedData->bShouldUpdatedSelectedCoMs = false;
-	}
-}
-
 #undef LOCTEXT_NAMESPACE
diff --git a/Engine/Source/Editor/PhysicsAssetEditor/Private/PhysicsAssetEditorEditMode.h b/Engine/Source/Editor/PhysicsAssetEditor/Private/PhysicsAssetEditorEditMode.h
--- a/Engine/Source/Editor/PhysicsAssetEditor/Private/PhysicsAssetEditorEditMode.h
+++ b/Engine/Source/Editor/PhysicsAssetEditor/Private/PhysicsAssetEditorEditMode.h
@@ -93,9 +93,6 @@ private:
 	bool HandleBeginTransform();
 	bool HandleEndTransform(FEditorViewportClient* InViewportClient) const;
 
-	// Update the Center of Mass position after body transforms have been manipulated in the view port.
-	void UpdateCoM();
-
 private:
 	/** Shared data */
 	FPhysicsAssetEditorSharedData* SharedData;
@@ -121,10 +118,11 @@ private:
 	FVector SimGrabZ;
 
 	/** Members used for interacting with the asset while the simulation is running */
+	TArray<FTransform> ManConTM;
 	TArray<FTransform> StartManRelConTM;
 	TArray<FTransform> StartManParentConTM;
 	TArray<FTransform> StartManChildConTM;
-
+	
 	float DragX;
 	float DragY;
 
diff --git a/Engine/Source/Editor/PhysicsAssetEditor/Private/PhysicsAssetEditorSelection.cpp b/Engine/Source/Editor/PhysicsAssetEditor/Private/PhysicsAssetEditorSelection.cpp
new file mode 100644
--- /dev/null
+++ b/Engine/Source/Editor/PhysicsAssetEditor/Private/PhysicsAssetEditorSelection.cpp
@@ -0,0 +1,586 @@
+// Copyright Epic Games, Inc. All Rights Reserved.
+
+#include "PhysicsAssetEditorSelection.h"
+
+#include "Algo/Find.h"
+#include "Algo/ForEach.h"
+#include "PhysicsEngine/SkeletalBodySetup.h"
+#include "PhysicsEngine/PhysicsAsset.h"
+#include "ScopedTransaction.h"
+
+#define LOCTEXT_NAMESPACE "PhysicsAssetEditorSelection"
+
+// File Scope Utility Functions //
+
+template<typename CollectionType>
+int32 CountElementsOfType(CollectionType&& Collection, const uint32 InElementTypeFlags)
+{
+	int32 Count = 0;
+	Algo::ForEach(Collection, [&Count, InElementTypeFlags](const auto& Element) { if (Element.HasType(InElementTypeFlags)) { ++Count; } });
+	return Count;
+}
+
+template<typename CollectionType>
+bool ContainsElementsOfType(CollectionType&& Collection, const uint32 InElementTypeFlags)
+{
+	return Algo::FindByPredicate(Collection, [InElementTypeFlags](const auto& Element) { return Element.HasType(InElementTypeFlags); }) != nullptr;
+}
+
+bool SetSelected(TArray<FPhysicsAssetEditorSelectedElement>& InSelectedElements, const TArray<FPhysicsAssetEditorSelectedElement>& InDeltaElements, const bool bSelected)
+{
+	bool bSuccess = false;
+
+	if (!InDeltaElements.IsEmpty())
+	{
+		if (bSelected)
+		{
+			for (const FPhysicsAssetEditorSelectedElement& Element : InDeltaElements)
+			{
+				bSuccess |= InSelectedElements.AddUnique(Element) != INDEX_NONE;
+			}
+		}
+		else
+		{
+			for (const FPhysicsAssetEditorSelectedElement& Element : InDeltaElements)
+			{
+				bSuccess |= InSelectedElements.Remove(Element) != INDEX_NONE;
+			}
+		}
+	}
+
+	return bSuccess; // Return true if the selection was changed.
+}
+
+bool ModifySelection(TArray<FPhysicsAssetEditorSelectedElement>& DestinationSelection, const TArray<FPhysicsAssetEditorSelectedElement>& SourceSelection, const bool bSelected)
+{
+	if (SourceSelection.Num() == 0)
+	{
+		return false;
+	}
+
+	if (bSelected)
+	{
+		for (const FPhysicsAssetEditorSelectedElement& Element : SourceSelection)
+		{
+			DestinationSelection.AddUnique(Element);
+		}
+	}
+	else
+	{
+		for (const FPhysicsAssetEditorSelectedElement& Element : SourceSelection)
+		{
+			DestinationSelection.Remove(Element);
+		}
+	}
+
+	return true;
+}
+
+template< typename PrimitiveType> void InitializeSelectionWithFirstPrimitive(FPhysicsAssetEditorSelectedElement& Selection, const TArray<PrimitiveType>& PrimitiveElements)
+{
+	if (Selection.PrimitiveType == EAggCollisionShape::Unknown)
+	{
+		if (!PrimitiveElements.IsEmpty())
+		{
+			Selection.PrimitiveType = PrimitiveElements[0].GetShapeType();
+			Selection.PrimitiveIndex = 0;
+		}
+	}
+}
+
+void InitializeSelectionWithFirstPrimitive(FPhysicsAssetEditorSelectedElement& Selection, TObjectPtr<UPhysicsAsset> PhysicsAsset)
+{
+	const int BodyIndex = Selection.Index;
+
+	if (PhysicsAsset && PhysicsAsset->SkeletalBodySetups.IsValidIndex(BodyIndex))
+	{
+		const TObjectPtr<USkeletalBodySetup> BodySetup = PhysicsAsset->SkeletalBodySetups[BodyIndex];
+		check(BodySetup);
+		const FKAggregateGeom& AggGeom = BodySetup->AggGeom;
+
+		InitializeSelectionWithFirstPrimitive(Selection, AggGeom.SphereElems);
+		InitializeSelectionWithFirstPrimitive(Selection, AggGeom.BoxElems);
+		InitializeSelectionWithFirstPrimitive(Selection, AggGeom.SphylElems);
+		InitializeSelectionWithFirstPrimitive(Selection, AggGeom.ConvexElems);
+		InitializeSelectionWithFirstPrimitive(Selection, AggGeom.TaperedCapsuleElems);
+	}
+}
+
+template<typename PrimitiveType> void MakeSelectionForEachPrimitive(const int32 InBodyIndex, const TArray<PrimitiveType>& InPrimitiveElements, TArray<FPhysicsAssetEditorSelectedElement>& OutSelectedElements)
+{
+	int32 PrimitiveIndex = 0;
+
+	Algo::ForEach(InPrimitiveElements,
+		[InBodyIndex, &OutSelectedElements, &PrimitiveIndex](const PrimitiveType& Primitive)
+		{
+			OutSelectedElements.Add(MakePrimitiveSelection(InBodyIndex, Primitive.GetShapeType(), PrimitiveIndex++));
+		});
+}
+
+
+template<typename CollectionType> FString BuildSelectionDescriptionText(CollectionType&& InSelectedElements)
+{
+	int32 SelectedBodyCount = 0;
+	int32 SelectedCenterOfMassCount = 0;
+	int32 SelectedConstraintCount = 0;
+	int32 SelectedPrimitiveCount = 0;
+
+	// Count instances of each type in selection.
+	for (const FPhysicsAssetEditorSelectedElement& SelectedElement : InSelectedElements)
+	{
+		if (SelectedElement.SelectedElementType == FPhysicsAssetEditorSelectedElement::Body) ++SelectedBodyCount;
+		if (SelectedElement.SelectedElementType == FPhysicsAssetEditorSelectedElement::CenterOfMass) ++SelectedCenterOfMassCount;
+		if (SelectedElement.SelectedElementType == FPhysicsAssetEditorSelectedElement::Constraint) ++SelectedConstraintCount;
+		if (SelectedElement.SelectedElementType == FPhysicsAssetEditorSelectedElement::Primitive) ++SelectedPrimitiveCount;
+	}
+
+	FString Description;
+
+	// Build String summarizing the number of each type of selected element.
+	if (SelectedBodyCount > 0) Description += FText::Format(INVTEXT("{0} {0}|plural(one=Body, other=Bodies) "), SelectedBodyCount).ToString();
+	if (SelectedCenterOfMassCount > 0)  Description += FText::Format(INVTEXT("{0} {0}|plural(one=CoM, other=CoMs) "), SelectedCenterOfMassCount).ToString();
+	if (SelectedConstraintCount > 0)  Description += FText::Format(INVTEXT("{0} {0}|plural(one=Constraint, other=Constraints) "), SelectedConstraintCount).ToString();
+	if (SelectedPrimitiveCount > 0)  Description += FText::Format(INVTEXT("{0} {0}|plural(one=Primitive, other=Primitives) "), SelectedPrimitiveCount).ToString();
+
+	
+	if (!Description.IsEmpty())
+	{
+		// Remove trailing comma from description.
+		Description.TrimEndInline();
+		Description.TrimCharInline(',', nullptr);
+	
+		// Replace final comma with ' and'.
+		int32 Location = INDEX_NONE;
+		Description.FindLastChar(',', Location);
+		if (Location != INDEX_NONE)
+		{
+			Description.RemoveAt(Location);
+			Description.InsertAt(Location, LOCTEXT("and", " and").ToString());
+		}
+	}
+	else
+	{
+		Description = LOCTEXT("None", "none").ToString();
+	}
+
+	return Description;
+}
+
+namespace TransactionText
+{
+	static const FText SelectText = LOCTEXT("Select", "Select");
+	static const FText UnselectText = LOCTEXT("Unselect", "Unselect");
+
+	FText GetModifySelectionText(const bool bIsSelecting) // TODO - better name for this
+	{
+		return bIsSelecting ? TransactionText::SelectText : TransactionText::UnselectText;
+	}
+}
+
+int32 FindNextElementOfType(const UPhysicsAssetEditorSelection& Selection, const int32 InElementIndex, const uint32 InElementTypeFlags, const bool bSearchInReverse)
+{
+	int32 Index = InElementIndex;
+	const FPhysicsAssetEditorSelectedElement* Element;
+
+	do
+	{
+		Index = bSearchInReverse ? --Index : ++Index;
+		Element = Selection.GetSelectedAtValidIndex(Index);
+	} while (Element && !Element->HasType(InElementTypeFlags));
+	
+	return FMath::Clamp(Index, 0, Selection.Num()); // Clamp to range.
+}
+
+// Externally Accessible Functions //
+
+bool IsReferencingPrimitive(const FPhysicsAssetEditorSelectedElement& InSelection, const int32 InBodyIndex, const EAggCollisionShape::Type InPrimitiveType, const int32 InPrimitiveIndex)
+{
+	return (InBodyIndex == InSelection.GetIndex()) && (InPrimitiveType == InSelection.GetPrimitiveType()) && (InPrimitiveIndex == InSelection.GetPrimitiveIndex());
+}
+
+FPhysicsAssetEditorSelectedElement MakeBodySelection(TObjectPtr<UPhysicsAsset> PhysicsAsset, const int32 BodyIndex)
+{
+	FPhysicsAssetEditorSelectedElement Selection(FPhysicsAssetEditorSelectedElement::Body, BodyIndex);
+	InitializeSelectionWithFirstPrimitive(Selection, PhysicsAsset); // We choose a primitive in the body s.t. we can position a widget in the view port - doesn't feel right.
+	return Selection;
+}
+
+TArray<FPhysicsAssetEditorSelectedElement> MakeBodySelection(TObjectPtr<UPhysicsAsset> InPhysicsAsset, const TArray<int32> InBodyIndexCollection)
+{
+	TArray<FPhysicsAssetEditorSelectedElement> SelectedElements;
+	SelectedElements.Reserve(InBodyIndexCollection.Num());
+
+	for (const int32 BodyIndex : InBodyIndexCollection)
+	{
+		SelectedElements.Add(MakeBodySelection(InPhysicsAsset, BodyIndex));
+	}
+
+	return SelectedElements;
+}
+
+FPhysicsAssetEditorSelectedElement MakeSelectionAnyPrimitiveInBody(TObjectPtr<UPhysicsAsset> InPhysicsAsset, const int32 InBodyIndex)
+{
+	FPhysicsAssetEditorSelectedElement Selection(FPhysicsAssetEditorSelectedElement::Primitive, InBodyIndex);
+	InitializeSelectionWithFirstPrimitive(Selection, InPhysicsAsset);
+	return Selection;
+}
+
+TArray<FPhysicsAssetEditorSelectedElement> MakeSelectionAllPrimitivesInBody(TObjectPtr<UPhysicsAsset> InPhysicsAsset, const int32 InBodyIndex)
+{
+	return MakeSelectionAllPrimitivesInBody(InPhysicsAsset, TArray<int32>{ InBodyIndex });
+}
+
+TArray<FPhysicsAssetEditorSelectedElement> MakeSelectionAllPrimitivesInBody(TObjectPtr<UPhysicsAsset> InPhysicsAsset, const TArray<int32> InBodyIndexCollection)
+{
+	TArray<FPhysicsAssetEditorSelectedElement> SelectedElements;
+
+	for (const int32 BodyIndex : InBodyIndexCollection)
+	{
+		if (InPhysicsAsset->SkeletalBodySetups.IsValidIndex(BodyIndex))
+		{
+			const TObjectPtr<USkeletalBodySetup> BodySetup = InPhysicsAsset->SkeletalBodySetups[BodyIndex];
+			check(BodySetup);
+			const FKAggregateGeom& AggGeom = BodySetup->AggGeom;
+
+			MakeSelectionForEachPrimitive(BodyIndex, AggGeom.SphereElems, SelectedElements);
+			MakeSelectionForEachPrimitive(BodyIndex, AggGeom.BoxElems, SelectedElements);
+			MakeSelectionForEachPrimitive(BodyIndex, AggGeom.SphylElems, SelectedElements);
+			MakeSelectionForEachPrimitive(BodyIndex, AggGeom.ConvexElems, SelectedElements);
+			MakeSelectionForEachPrimitive(BodyIndex, AggGeom.TaperedCapsuleElems, SelectedElements);
+		}
+	}
+
+	return SelectedElements;
+}
+
+
+FPhysicsAssetEditorSelectedElement MakeCoMSelection(const int32 InCoMIndex)
+{
+	return FPhysicsAssetEditorSelectedElement(FPhysicsAssetEditorSelectedElement::CenterOfMass, InCoMIndex);
+}
+
+FPhysicsAssetEditorSelectedElement MakeConstraintSelection(const int32 InConstraintIndex)
+{
+	return FPhysicsAssetEditorSelectedElement(FPhysicsAssetEditorSelectedElement::Constraint, InConstraintIndex);
+}
+
+TArray<FPhysicsAssetEditorSelectedElement> MakeConstraintSelection(const TArray<int32>& Indices)
+{
+	TArray<FPhysicsAssetEditorSelectedElement> Selection;
+	Selection.Reserve(Indices.Num());
+
+	for (const int32 Index : Indices)
+	{
+		Selection.Add(MakeConstraintSelection(Index));
+	}
+
+	return Selection;
+}
+
+FPhysicsAssetEditorSelectedElement MakePrimitiveSelection(const int32 InBodyIndex, const EAggCollisionShape::Type InPrimitiveType, const int32 InPrimitiveIndex)
+{
+	return FPhysicsAssetEditorSelectedElement(FPhysicsAssetEditorSelectedElement::Primitive, InBodyIndex, InPrimitiveType, InPrimitiveIndex);
+}
+
+// class FPhysicsAssetEditorSelectedElement //
+
+FPhysicsAssetEditorSelectedElement::FPhysicsAssetEditorSelectedElement(const Type InSelectedElementType, const int32 InBodyIndex, const EAggCollisionShape::Type InPrimitiveType, const int32 InPrimitiveIndex)
+	: Index(InBodyIndex)
+	, PrimitiveType(InPrimitiveType)
+	, PrimitiveIndex(InPrimitiveIndex)
+	, SelectedElementType(InSelectedElementType)
+{}
+
+FPhysicsAssetEditorSelectedElement::FPhysicsAssetEditorSelectedElement(const Type InSelectedElementType, const int32 InBodyIndex)
+	: FPhysicsAssetEditorSelectedElement(InSelectedElementType, InBodyIndex, EAggCollisionShape::Unknown, INDEX_NONE)
+{}
+
+int32 FPhysicsAssetEditorSelectedElement::GetIndex() const 
+{ 
+	return Index; 
+}
+
+EAggCollisionShape::Type FPhysicsAssetEditorSelectedElement::GetPrimitiveType() const 
+{
+	return EAggCollisionShape::Type(PrimitiveType);
+}
+
+int32 FPhysicsAssetEditorSelectedElement::GetPrimitiveIndex() const
+{
+	return PrimitiveIndex;
+}
+
+FPhysicsAssetEditorSelectedElement::Type FPhysicsAssetEditorSelectedElement::GetElementType() const
+{
+	return Type(SelectedElementType);
+}
+
+bool FPhysicsAssetEditorSelectedElement::HasType(const uint32 InElementTypeFlags) const
+{
+	return SelectedElementType & InElementTypeFlags;
+}
+
+bool operator==(const FPhysicsAssetEditorSelectedElement& Lhs, const FPhysicsAssetEditorSelectedElement& Rhs)
+{
+	return Lhs.Index == Rhs.Index && Lhs.PrimitiveType == Rhs.PrimitiveType && Lhs.PrimitiveIndex == Rhs.PrimitiveIndex && Lhs.SelectedElementType == Rhs.SelectedElementType;
+}
+
+// class FPhysicsAssetEditorSelectionIterator //
+
+FPhysicsAssetEditorSelectionIterator::FPhysicsAssetEditorSelectionIterator()
+	: FPhysicsAssetEditorSelectionIterator(nullptr, FPhysicsAssetEditorSelectedElement::All)
+{}
+
+FPhysicsAssetEditorSelectionIterator::FPhysicsAssetEditorSelectionIterator(const UPhysicsAssetEditorSelection* const InParentSelection)
+	: FPhysicsAssetEditorSelectionIterator(InParentSelection, FPhysicsAssetEditorSelectedElement::All)
+{}
+
+FPhysicsAssetEditorSelectionIterator::FPhysicsAssetEditorSelectionIterator(const UPhysicsAssetEditorSelection* const InParentSelection, const uint32 InElementTypeFlags)
+	: FPhysicsAssetEditorSelectionIterator(InParentSelection, InElementTypeFlags, 0)
+{}
+
+FPhysicsAssetEditorSelectionIterator::FPhysicsAssetEditorSelectionIterator(const UPhysicsAssetEditorSelection* const InParentSelection, const uint32 InElementTypeFlags, const uint32 InFilteredElementIndex)
+	: ParentSelection(InParentSelection)
+	, ElementTypeFlags(InElementTypeFlags)
+	, FilteredElementIndex(0)
+	, ParentElementIndex(0)
+{
+	check(ParentSelection != nullptr);
+
+	// Find first element of the filtered type.
+	while (IsValid() && !GetElement().HasType(ElementTypeFlags))
+	{
+		ParentElementIndex = FindNextElementOfType(*ParentSelection, ParentElementIndex, ElementTypeFlags, false);
+	}
+	
+	// Increment Index to reference the n-th element of the filtered type.
+	while (IsValid() && (FilteredElementIndex < InFilteredElementIndex)) 
+	{
+		StepIndex(false);
+	}
+}
+
+FPhysicsAssetEditorSelectionIterator::FPhysicsAssetEditorSelectionIterator(const FPhysicsAssetEditorSelectionIterator& Other)
+	: ParentSelection(Other.ParentSelection)
+	, ElementTypeFlags(Other.ElementTypeFlags)
+	, FilteredElementIndex(Other.FilteredElementIndex)
+	, ParentElementIndex(Other.ParentElementIndex)
+{}
+
+FPhysicsAssetEditorSelectionIterator::operator bool() const
+{
+	return IsValid();
+}
+
+uint32 FPhysicsAssetEditorSelectionIterator::GetElementTypeFlags() const 
+{
+	return ElementTypeFlags;
+}
+
+uint32 FPhysicsAssetEditorSelectionIterator::GetIndexIntoFilteredSelection() const
+{
+	return FilteredElementIndex;
+}
+
+uint32 FPhysicsAssetEditorSelectionIterator::GetIndexIntoParentSelection() const
+{
+	return ParentElementIndex;
+}
+
+const UPhysicsAssetEditorSelection* const FPhysicsAssetEditorSelectionIterator::GetParentSelection() const
+{
+	return ParentSelection;
+}
+
+void FPhysicsAssetEditorSelectionIterator::StepIndex(const bool bReverseDirection)
+{
+	const uint32 NextParentElementIndex = FindNextElementOfType(*ParentSelection, ParentElementIndex, ElementTypeFlags, bReverseDirection);
+
+	if (ParentElementIndex != NextParentElementIndex)
+	{
+		ParentElementIndex = NextParentElementIndex;
+		
+		if (bReverseDirection)
+		{
+			--FilteredElementIndex;
+		}
+		else
+		{
+			++FilteredElementIndex;
+		}
+	}
+}
+
+const FPhysicsAssetEditorSelectionIterator::Element& FPhysicsAssetEditorSelectionIterator::operator*()
+{
+	return GetElement();
+}
+
+const FPhysicsAssetEditorSelectionIterator::Element* FPhysicsAssetEditorSelectionIterator::operator->()
+{
+	return &GetElement();
+}
+
+bool FPhysicsAssetEditorSelectionIterator::IsValid() const
+{
+	return ParentSelection->SelectedElements().IsValidIndex(ParentElementIndex);
+}
+
+const FPhysicsAssetEditorSelectionIterator::Element& FPhysicsAssetEditorSelectionIterator::GetElement()
+{
+	return ParentSelection->GetSelectedAt(ParentElementIndex);
+}
+
+void FPhysicsAssetEditorSelectionIterator::operator++()
+{
+	StepIndex(false);
+}
+
+void FPhysicsAssetEditorSelectionIterator::operator--()
+{
+	StepIndex(true);
+}
+
+bool FPhysicsAssetEditorSelectionIterator::operator==(const FPhysicsAssetEditorSelectionIterator & Other) const
+{
+	return (ElementTypeFlags == Other.ElementTypeFlags)
+		&& (ParentElementIndex == Other.ParentElementIndex)
+		&& (FilteredElementIndex == Other.FilteredElementIndex)
+		&& (ParentSelection == Other.ParentSelection);
+}
+
+// class UPhysicsAssetEditorSelection //
+
+const TArray<FPhysicsAssetEditorSelectedElement>& UPhysicsAssetEditorSelection::SelectedElements() const
+{
+	return SelectedElementCollection;
+}
+
+const UPhysicsAssetEditorSelection::Element& UPhysicsAssetEditorSelection::GetSelectedAt(const int32 InElementIndex) const
+{
+	check(IsValidIndex(InElementIndex));
+
+	return SelectedElementCollection[InElementIndex];
+}
+
+const UPhysicsAssetEditorSelection::Element* UPhysicsAssetEditorSelection::GetSelectedAtValidIndex(const int32 InElementIndex) const
+{
+	return (IsValidIndex(InElementIndex)) ? &GetSelectedAt(InElementIndex) : nullptr;
+}
+
+bool UPhysicsAssetEditorSelection::IsValidIndex(const int32 InElementIndex) const
+{
+	return SelectedElementCollection.IsValidIndex(InElementIndex);
+}
+
+bool UPhysicsAssetEditorSelection::ContainsType(const uint32 InElementTypeFlags) const
+{
+	return ContainsElementsOfType(SelectedElementCollection, InElementTypeFlags);
+}
+
+UPhysicsAssetEditorSelection::FilterRange UPhysicsAssetEditorSelection::SelectedElementsOfType(const uint32 InElementTypeFlags) const
+{	
+	FilterIterator Begin(this, InElementTypeFlags);
+	FilterIterator End(this, InElementTypeFlags, NumOfType(InElementTypeFlags));
+
+	return FilterRange(Begin, End);
+}
+
+UPhysicsAssetEditorSelection::UniqueRange UPhysicsAssetEditorSelection::UniqueSelectedElementsOfType(const uint32 InElementTypeFlags) const
+{
+	FilterIterator FilterItertorBegin(this, InElementTypeFlags);
+	UniqueIterator Begin(FilterItertorBegin);
+	UniqueIterator End(FilterItertorBegin);
+
+	for (; End.IsValid(); ++End); // Advance iterator to the end of the unique range
+
+	return UniqueRange(Begin, End);
+}
+
+bool UPhysicsAssetEditorSelection::ModifySelected(const TArray<FPhysicsAssetEditorSelectedElement>& InSelectedElements, const bool bSelected)
+{
+	const FText TransactionFormat = (bSelected) ? LOCTEXT("AddToTheCurrentSelection", "Add {0} to the current selection") : LOCTEXT("RemoveFromTheCurrentSelection", "Remove {0} from the current selection");
+	const FScopedTransaction Transaction(FText::Format(TransactionFormat, FText::FromString(BuildSelectionDescriptionText(InSelectedElements))));
+	Modify();
+	return ModifySelection(SelectedElementCollection, InSelectedElements, bSelected);
+}
+
+bool UPhysicsAssetEditorSelection::SetSelected(const TArray<FPhysicsAssetEditorSelectedElement>& InSelectedElements)
+{
+	const FScopedTransaction Transaction(FText::Format(LOCTEXT("Set selection", "Set selection to {0}"), FText::FromString(BuildSelectionDescriptionText(InSelectedElements))));
+	Modify();
+	ClearSelectionInternal();
+	return ModifySelection(SelectedElementCollection, InSelectedElements, true);
+}
+
+const FPhysicsAssetEditorSelectedElement* UPhysicsAssetEditorSelection::GetNextSelectedOfType(const int32 InElementIndex, const uint32 InElementTypeFlags) const
+{
+	for (int32 Index = InElementIndex; SelectedElementCollection.IsValidIndex(Index); ++Index)
+	{
+		if (SelectedElementCollection[Index].HasType(InElementTypeFlags))
+		{
+			return &SelectedElementCollection[Index];
+		}
+	}
+
+	return nullptr;
+}
+
+void UPhysicsAssetEditorSelection::ClearSelection()
+{
+	if (!SelectedElementCollection.IsEmpty())
+	{
+		const FScopedTransaction Transaction(FText::Format(LOCTEXT("ClearSelection", "Clear Selection ({0})"), FText::FromString(BuildSelectionDescriptionText(SelectedElementCollection))));
+		Modify();
+		ClearSelectionInternal();
+	}
+}
+
+void UPhysicsAssetEditorSelection::ClearSelection(const uint32 InElementTypeFlags)
+{
+	ClearSelectionInternal(InElementTypeFlags, true);
+}
+
+void UPhysicsAssetEditorSelection::ClearSelectionWithoutTransaction(const uint32 InElementTypeFlags)
+{
+	ClearSelectionInternal(InElementTypeFlags, false);
+}
+
+int32 UPhysicsAssetEditorSelection::NumOfType(const uint32 InElementTypeFlags) const
+{
+	return CountElementsOfType(SelectedElementCollection, InElementTypeFlags);
+}
+
+const FPhysicsAssetEditorSelectedElement* UPhysicsAssetEditorSelection::GetLastSelectedOfType(const uint32 InElementTypeFlags) const
+{
+	const int32 LastSelectedElementIndex = SelectedElementCollection.FindLastByPredicate([InElementTypeFlags](const FPhysicsAssetEditorSelectedElement& Element) { return Element.SelectedElementType & InElementTypeFlags; });
+	
+	if (SelectedElementCollection.IsValidIndex(LastSelectedElementIndex))
+	{
+		return &SelectedElementCollection[LastSelectedElementIndex];
+	}
+
+	return nullptr;
+}
+
+void UPhysicsAssetEditorSelection::ClearSelectionInternal()
+{
+	SelectedElementCollection.Reset();
+}
+
+void UPhysicsAssetEditorSelection::ClearSelectionInternal(const uint32 InElementTypeFlags, const bool bShouldCreateTransaction)
+{
+	FilterRange ElementsToRemove = SelectedElementsOfType(InElementTypeFlags);
+
+	if (!ElementsToRemove.IsEmpty())
+	{
+		const FScopedTransaction Transaction(FText::Format(LOCTEXT("ClearSelection", "Clear Selection ({0})"), FText::FromString(BuildSelectionDescriptionText(ElementsToRemove))), bShouldCreateTransaction);
+		Modify();
+
+		TArray<FPhysicsAssetEditorSelectedElement> ElementsToRemoveArray;
+		Algo::Copy(ElementsToRemove, ElementsToRemoveArray); // Copy range elements into an array as range may reference the DestinationSelection.
+		ModifySelection(SelectedElementCollection, ElementsToRemoveArray, false);
+	}
+}
+
+#undef LOCTEXT_NAMESPACE 
\ No newline at end of file
diff --git a/Engine/Source/Editor/PhysicsAssetEditor/Private/PhysicsAssetEditorSelection.h b/Engine/Source/Editor/PhysicsAssetEditor/Private/PhysicsAssetEditorSelection.h
new file mode 100644
--- /dev/null
+++ b/Engine/Source/Editor/PhysicsAssetEditor/Private/PhysicsAssetEditorSelection.h
@@ -0,0 +1,326 @@
+// Copyright Epic Games, Inc. All Rights Reserved.
+
+#pragma once
+
+#include "CoreMinimal.h"
+#include "PhysicsEngine/ShapeElem.h"
+#include "Preferences/PhysicsAssetEditorOptions.h"
+
+#include "PhysicsAssetEditorSelection.generated.h"
+
+class UPhysicsAssetEditorSelection;
+
+/** Encapsulates a selected set of bodies or constraints */
+
+USTRUCT()
+struct FPhysicsAssetEditorSelectedElement
+{
+	GENERATED_BODY()
+
+	enum Type : int32
+	{
+		None,
+
+		Body			= 1,		// A physics body, comprised of one or more primitives.
+		CenterOfMass	= 1 << 1,	// A Center of Mass of a physics body.
+		Constraint		= 1 << 2,	// A constraint between two physics bodies.
+		Primitive		= 1 << 3,	// A primitive that defines the geometry of a physics body.
+
+		All				= ~None
+	};
+
+	FPhysicsAssetEditorSelectedElement() = default;
+	FPhysicsAssetEditorSelectedElement(const Type InSelectedElementType, const int32 InBodyIndex, const EAggCollisionShape::Type InPrimitiveType, const int32 InPrimitiveIndex);
+	FPhysicsAssetEditorSelectedElement(const Type InSelectedElementType, const int32 InBodyIndex);
+
+	int32 GetIndex() const;
+	EAggCollisionShape::Type GetPrimitiveType() const;
+	int32 GetPrimitiveIndex() const;
+	Type GetElementType() const;
+
+	bool HasType(const uint32 InElementTypeFlags) const;
+
+	UPROPERTY()
+	int32 Index;
+
+	UPROPERTY()
+	TEnumAsByte<EAggCollisionShape::Type> PrimitiveType = EAggCollisionShape::Unknown;
+	
+	UPROPERTY()
+	int32 PrimitiveIndex = INDEX_NONE;
+
+	UPROPERTY()
+	int32 SelectedElementType = Type::None;
+
+};
+
+bool operator==(const FPhysicsAssetEditorSelectedElement& Lhs, const FPhysicsAssetEditorSelectedElement& Rhs);
+
+bool IsReferencingPrimitive(const FPhysicsAssetEditorSelectedElement& InSelection, const int32 InBodyIndex, const EAggCollisionShape::Type InPrimitiveType, const int32 InPrimitiveIndex);
+
+FPhysicsAssetEditorSelectedElement MakeBodySelection(TObjectPtr<UPhysicsAsset> InPhysicsAsset, const int32 InBodyIndex);
+TArray<FPhysicsAssetEditorSelectedElement> MakeBodySelection(TObjectPtr<UPhysicsAsset> InPhysicsAsset, const TArray<int32> InBodyIndexCollection);
+
+FPhysicsAssetEditorSelectedElement MakeCoMSelection(const int32 InCoMIndex);
+
+FPhysicsAssetEditorSelectedElement MakeConstraintSelection(const int32 InConstraintIndex);
+TArray<FPhysicsAssetEditorSelectedElement> MakeConstraintSelection(const TArray<int32>& Indices);
+
+FPhysicsAssetEditorSelectedElement MakePrimitiveSelection(const int32 InBodyIndex, const EAggCollisionShape::Type InPrimitiveType, const int32 InPrimitiveIndex);
+
+FPhysicsAssetEditorSelectedElement MakeSelectionAnyPrimitiveInBody(TObjectPtr<UPhysicsAsset> InPhysicsAsset, const int32 InBodyIndex); // Selected the 'first' primitive in the supplied body
+TArray<FPhysicsAssetEditorSelectedElement> MakeSelectionAllPrimitivesInBody(TObjectPtr<UPhysicsAsset> InPhysicsAsset, const int32 InBodyIndex);
+TArray<FPhysicsAssetEditorSelectedElement> MakeSelectionAllPrimitivesInBody(TObjectPtr<UPhysicsAsset> InPhysicsAsset, const TArray<int32> InBodyIndexCollection);
+
+
+class FPhysicsAssetEditorSelectionIterator
+{
+public:
+
+	using Element = FPhysicsAssetEditorSelectedElement;
+	
+	FPhysicsAssetEditorSelectionIterator();
+	FPhysicsAssetEditorSelectionIterator(const UPhysicsAssetEditorSelection* const InParentSelection);
+	FPhysicsAssetEditorSelectionIterator(const UPhysicsAssetEditorSelection* const InParentSelection, const uint32 InElementTypeFlags);
+	FPhysicsAssetEditorSelectionIterator(const UPhysicsAssetEditorSelection* const InParentSelection, const uint32 InElementTypeFlags, const uint32 InElementIndex);
+	FPhysicsAssetEditorSelectionIterator(const FPhysicsAssetEditorSelectionIterator& Other);
+
+	bool IsValid() const;
+	operator bool() const;
+
+	const Element& operator*();
+	const Element* operator->();
+
+	void operator++();
+	void operator--();
+	bool operator==(const FPhysicsAssetEditorSelectionIterator& Other) const;
+
+	uint32 GetElementTypeFlags() const;
+	uint32 GetIndexIntoFilteredSelection() const;
+	uint32 GetIndexIntoParentSelection() const;
+	const UPhysicsAssetEditorSelection* const GetParentSelection() const;
+
+private:
+	void StepIndex(const bool bReverseDirection);
+
+	const Element& GetElement();
+
+	const UPhysicsAssetEditorSelection* const ParentSelection;
+	const uint32 ElementTypeFlags;
+
+	uint32 FilteredElementIndex;
+	uint32 ParentElementIndex;
+};
+
+template< typename IteratorType > class TPhysicsAssetEditorUniqueIterator
+{
+public:
+
+	using ThisType = TPhysicsAssetEditorUniqueIterator< IteratorType >;
+	using Iterator = IteratorType;
+	using Element = typename Iterator::Element;
+
+	TPhysicsAssetEditorUniqueIterator(Iterator InIterator);
+
+	bool IsValid() const;
+	operator bool() const;
+	const Element& operator*();
+	const Element* operator->();
+
+	void operator++();
+	void operator--();
+	bool operator==(const ThisType& Other) const;
+
+private:
+	TArray<int32> EncounteredIndexIntoParentSelection;
+	IteratorType ManagedIterator;
+};
+
+template< typename IteratorType > class TPhysicsAssetEditorSelectionRange
+{
+public:
+	using Iterator = IteratorType;
+	using Element = Iterator::Element;
+	using TConstIterator = Iterator;
+
+	TPhysicsAssetEditorSelectionRange(Iterator InBegin, Iterator InEnd);
+
+	Iterator CreateConstIterator();
+
+	Iterator begin() const;
+	Iterator end() const;
+
+	bool IsValid() const;
+	operator bool() const;
+	bool IsEmpty() const;
+	int32 Num() const;
+	TArray<Element> ToArray();
+	const Element& At(const int32 InIndex) const;
+
+private:
+	Iterator IteratorBegin;
+	Iterator IteratorEnd;
+};
+
+UCLASS()
+class UPhysicsAssetEditorSelection : public UObject
+{
+public:
+	GENERATED_BODY()
+
+	using Element = FPhysicsAssetEditorSelectedElement;
+	using Selection = TArray< Element >;
+	using Iterator = FPhysicsAssetEditorSelectionIterator;
+
+	using FilterIterator = FPhysicsAssetEditorSelectionIterator;
+	using UniqueIterator = TPhysicsAssetEditorUniqueIterator< FilterIterator >;
+
+	using FilterRange = TPhysicsAssetEditorSelectionRange< FilterIterator >;
+	using UniqueRange = TPhysicsAssetEditorSelectionRange< UniqueIterator >;
+
+	const Selection& SelectedElements() const;
+	const Element& GetSelectedAt(const int32 InElementIndex) const;
+	const Element* GetSelectedAtValidIndex(const int32 InElementIndex) const;
+	bool IsValidIndex(const int32 InElementIndex) const;
+	bool ContainsType(const uint32 InElementTypeFlags) const;
+	int32 Num() const { return SelectedElementCollection.Num(); }
+	int32 NumOfType(const uint32 InElementTypeFlags) const;
+
+	void ClearSelection();
+	void ClearSelection(const uint32 InElementTypeFlags);
+	void ClearSelectionWithoutTransaction(const uint32 InElementTypeFlags); // This can be used to clear selected elements of type when inside an existing transaction.
+
+	Iterator SelectedElementsOfTypeIterator(const uint32 InElementTypeFlags) const { return Iterator(this, InElementTypeFlags); }
+	FilterRange SelectedElementsOfType(const uint32 InElementTypeFlags) const;
+	UniqueRange UniqueSelectedElementsOfType(const uint32 InElementTypeFlags) const;
+
+	bool ModifySelected(const TArray<Element>& InSelectedElements, const bool bSelected);
+	bool SetSelected(const TArray<Element>& InSelectedElements);
+
+	const Element* GetNextSelectedOfType(const int32 InElementIndex, const uint32 InElementTypeFlags) const;
+	const Element* GetLastSelectedOfType(const uint32 InElementTypeFlags) const;
+
+private:
+	
+	void ClearSelectionInternal();
+	void ClearSelectionInternal(const uint32 InElementTypeFlags, const bool bShouldCreateTransaction);
+
+	UPROPERTY()
+	TArray< FPhysicsAssetEditorSelectedElement > SelectedElementCollection;
+};
+
+// class TPhysicsAssetEditorUniqueIterator //
+
+template< typename IteratorType > TPhysicsAssetEditorUniqueIterator< IteratorType >::TPhysicsAssetEditorUniqueIterator(Iterator InIterator)
+	: ManagedIterator(InIterator)
+{}
+
+template< typename IteratorType > bool TPhysicsAssetEditorUniqueIterator< IteratorType >::IsValid() const
+{
+	return ManagedIterator.IsValid();
+}
+
+template< typename IteratorType > TPhysicsAssetEditorUniqueIterator< IteratorType >::operator bool() const
+{
+	return ManagedIterator.IsValid();
+}
+
+template< typename IteratorType > const typename TPhysicsAssetEditorUniqueIterator< IteratorType >::Element& TPhysicsAssetEditorUniqueIterator< IteratorType >::operator*()
+{
+	return ManagedIterator.operator*();
+}
+
+template< typename IteratorType > const typename TPhysicsAssetEditorUniqueIterator< IteratorType >::Element* TPhysicsAssetEditorUniqueIterator< IteratorType >::operator->()
+{
+	return ManagedIterator.operator->();
+}
+
+template< typename IteratorType > void TPhysicsAssetEditorUniqueIterator< IteratorType >::operator++()
+{
+	EncounteredIndexIntoParentSelection.Add(ManagedIterator.GetIndexIntoParentSelection());
+
+	while (IsValid() && EncounteredIndexIntoParentSelection.ContainsByPredicate([this](const int32& Element) { return ManagedIterator.GetParentSelection()->GetSelectedAt(Element).Index == ManagedIterator->Index; }))
+	{
+		++ManagedIterator;
+	}
+}
+
+template< typename IteratorType > void TPhysicsAssetEditorUniqueIterator< IteratorType >::operator--()
+{
+	if (!EncounteredIndexIntoParentSelection.IsEmpty())
+	{
+		// Step back to the last encountered element index.
+		while (ManagedIterator.GetIndexIntoParentSelection() > EncounteredIndexIntoParentSelection.Last())
+		{
+			--ManagedIterator;
+		}
+
+		EncounteredIndexIntoParentSelection.Pop();
+	}
+}
+
+template< typename IteratorType > bool TPhysicsAssetEditorUniqueIterator< IteratorType >::operator==(const ThisType& Other) const
+{
+	return (ManagedIterator == Other.ManagedIterator)
+		&& (EncounteredIndexIntoParentSelection == Other.EncounteredIndexIntoParentSelection);
+}
+
+
+// class TPhysicsAssetEditorSelectionRange //
+
+template< typename IteratorType > TPhysicsAssetEditorSelectionRange< IteratorType >::TPhysicsAssetEditorSelectionRange(Iterator InBegin, Iterator InEnd)
+	: IteratorBegin(InBegin)
+	, IteratorEnd(InEnd)
+{}
+
+template< typename IteratorType > TPhysicsAssetEditorSelectionRange< IteratorType >::Iterator TPhysicsAssetEditorSelectionRange< IteratorType >::CreateConstIterator()
+{
+	return begin();
+}
+
+template< typename IteratorType > TPhysicsAssetEditorSelectionRange< IteratorType >::Iterator TPhysicsAssetEditorSelectionRange< IteratorType >::begin() const
+{
+	return IteratorBegin;
+}
+
+template< typename IteratorType > TPhysicsAssetEditorSelectionRange< IteratorType >::Iterator TPhysicsAssetEditorSelectionRange< IteratorType >::end() const
+{
+	return IteratorEnd;
+}
+
+template< typename IteratorType > bool TPhysicsAssetEditorSelectionRange< IteratorType >::IsValid() const
+{
+	return begin() != end();
+}
+
+template< typename IteratorType > TPhysicsAssetEditorSelectionRange< IteratorType >::operator bool() const
+{
+	return IsValid();
+}
+
+template< typename IteratorType > bool TPhysicsAssetEditorSelectionRange< IteratorType >::IsEmpty() const
+{
+	return !IsValid();
+}
+
+template< typename IteratorType > int32 TPhysicsAssetEditorSelectionRange< IteratorType >::Num() const
+{
+	int32 Count = 0;
+	for (Iterator Itr = begin(); Itr != end(); ++Itr, ++Count);
+	return Count;
+};
+
+template< typename IteratorType > TArray< typename TPhysicsAssetEditorSelectionRange< IteratorType >::Element > TPhysicsAssetEditorSelectionRange< IteratorType >::ToArray()
+{
+	TArray<Element> Array;
+	Algo::Copy(*this, Array);
+	return Array;
+}
+
+template< typename IteratorType > const typename TPhysicsAssetEditorSelectionRange< IteratorType >::Element& TPhysicsAssetEditorSelectionRange< IteratorType >::At(const int32 InIndex) const
+{
+	Iterator Itr = begin();
+	for (int32 i = 0; i < InIndex; ++i, ++Itr);
+	return *Itr;
+}
+
diff --git a/Engine/Source/Editor/PhysicsAssetEditor/Private/PhysicsAssetEditorSharedData.cpp b/Engine/Source/Editor/PhysicsAssetEditor/Private/PhysicsAssetEditorSharedData.cpp
--- a/Engine/Source/Editor/PhysicsAssetEditor/Private/PhysicsAssetEditorSharedData.cpp
+++ b/Engine/Source/Editor/PhysicsAssetEditor/Private/PhysicsAssetEditorSharedData.cpp
@@ -4,6 +4,7 @@
 #include "Animation/MirrorDataTable.h"
 #include "PhysicsAssetEditorPhysicsHandleComponent.h"
 #include "PhysicsAssetRenderUtils.h"
+#include "PhysicsAssetEditorSelection.h"
 #include "PhysicsEngine/RigidBodyIndexPair.h"
 #include "Math/Axis.h"
 #include "Misc/MessageDialog.h"
@@ -56,44 +57,6 @@ namespace SharedDataConstants
 
 // File Scope Utility Functions //
 
-bool SetSelected(TArray<FPhysicsAssetEditorSharedData::FSelection>& SelectedElements, const TArray<FPhysicsAssetEditorSharedData::FSelection>& DeltaElements, const bool bSelected)
-{
-	if (DeltaElements.Num() == 0)
-	{
-		return false;
-	}
-
-	if (bSelected)
-	{
-		for (const FPhysicsAssetEditorSharedData::FSelection& Element : DeltaElements)
-		{
-			SelectedElements.AddUnique(Element);
-		}
-	}
-	else
-	{
-		for (const FPhysicsAssetEditorSharedData::FSelection& Element : DeltaElements)
-		{
-			SelectedElements.Remove(Element);
-		}
-	}
-
-	return true;
-}
-
-bool SetSelected(TArray<FPhysicsAssetEditorSharedData::FSelection>& SelectedElements, const TArray<int32>& DeltaElementIndexs, const bool bSelected)
-{
-	TArray<FPhysicsAssetEditorSharedData::FSelection> DeltaElements;
-	DeltaElements.Reserve(DeltaElementIndexs.Num());
-
-	for (const int32 ElementIndex : DeltaElementIndexs)
-	{
-		DeltaElements.Add(FPhysicsAssetEditorSharedData::FSelection(ElementIndex, EAggCollisionShape::Unknown, INDEX_NONE));
-	}
-
-	return SetSelected(SelectedElements, DeltaElements, bSelected);
-}
-
 /** Returns the Editor Body Flag bit mask that indicates if the supplied axis has been fixed in component space */
 int32 FindCoMAxisEditorBodyFlag(const EAxis::Type InAxis)
 {
@@ -110,7 +73,7 @@ namespace
 			const TShapeElem& ShapeElem = ShapeElems[PrimitiveIndex];
 			if (Predicate(SelectedElems, BodyIndex, ShapeElem))
 			{
-				SelectedElems.Add(FPhysicsAssetEditorSharedData::FSelection(BodyIndex, ShapeElem.GetShapeType(), PrimitiveIndex));
+				SelectedElems.Add(MakePrimitiveSelection(BodyIndex, ShapeElem.GetShapeType(), PrimitiveIndex));
 			}
 		}
 	}
@@ -137,11 +100,13 @@ TArray<FPhysicsAssetEditorSharedData::FSelection> CreateBodyPrimitivesSelection(
 
 // Pass each unique pair of values (excluding those containing the same value twice) in the supplied collection to the supplied function object.
 template<typename CollectionType, typename FunctionObjectType>
-void ForEachUniquePair(CollectionType& Collection, const FunctionObjectType& FunctionObject)
+void ForEachUniquePair(CollectionType&& Collection, const FunctionObjectType & FunctionObject)
 {
-	for (typename CollectionType::TConstIterator OuterItr = Collection.CreateConstIterator(); OuterItr; ++OuterItr)
+	using IteratorType = typename std::remove_reference<CollectionType>::type::TConstIterator;
+
+	for (IteratorType OuterItr = Collection.CreateConstIterator(); OuterItr; ++OuterItr)
 	{
-		typename CollectionType::TConstIterator InnerItr = OuterItr;
+		IteratorType InnerItr = OuterItr;
 		
 		++InnerItr;
 
@@ -152,9 +117,14 @@ void ForEachUniquePair(CollectionType& Collection, const FunctionObjectType& Fun
 	}
 }
 
-bool SelectionContainsIndex(const TArray<FPhysicsAssetEditorSharedData::FSelection>& InSelectedElements, const int32 InIndex)
+template<typename CollectionType>
+bool SelectionContainsIndex(CollectionType&& Collection, const int32 InIndex)
 {
-	return InSelectedElements.ContainsByPredicate([InIndex](const FPhysicsAssetEditorSharedData::FSelection& InSelection) {return InSelection.Index == InIndex; }); // Predicate returns true if the index is already in the selection.
+	return Algo::FindByPredicate(Collection, [InIndex](const FPhysicsAssetEditorSharedData::FSelection& InSelection) // Predicate returns true if the index is already in the selection.
+		{
+			return InSelection.Index == InIndex; 
+		}
+	) != nullptr;
 }
 
 // class FScopedBulkSelection //
@@ -196,8 +166,10 @@ FPhysicsAssetEditorSharedData::FPhysicsAssetEditorSharedData()
 	// Construct sim options.
 	EditorOptions = NewObject<UPhysicsAssetEditorOptions>(GetTransientPackage(), MakeUniqueObjectName(GetTransientPackage(), UPhysicsAssetEditorOptions::StaticClass(), FName(TEXT("EditorOptions"))), RF_Transactional);
 	check(EditorOptions);
-
 	EditorOptions->LoadConfig();
+
+	SelectedObjects = NewObject<UPhysicsAssetEditorSelection>(GetTransientPackage(), MakeUniqueObjectName(GetTransientPackage(), UPhysicsAssetEditorSelection::StaticClass(), FName(TEXT("PhysicsAssetEditorSelectedObjects"))), RF_Transactional);
+	check(SelectedObjects);
 }
 
 FPhysicsAssetEditorSharedData::~FPhysicsAssetEditorSharedData()
@@ -304,7 +276,7 @@ void FPhysicsAssetEditorSharedData::BroadcastSelectionChanged()
 {
 	if (!bSuspendSelectionBroadcast)
 	{
-		SelectionChangedEvent.Broadcast(SelectedBodies, SelectedConstraints, SelectedCoMs);
+		SelectionChangedEvent.Broadcast(SelectedObjects->SelectedElements());
 	}
 }
 
@@ -422,13 +394,16 @@ bool FPhysicsAssetEditorSharedData::GetShowCom() const
 
 FVector FPhysicsAssetEditorSharedData::GetCOMRenderPosition(const int32 BodyIndex) const
 {
-	if (bManipulating) 
+	if (IsManipulating())
 	{
-		if (const FSelection* const SelectedCoM = SelectedCoMs.FindByPredicate([BodyIndex](const FSelection& Element) { return Element.Index == BodyIndex; }))
+		if (SelectionContainsIndex(SelectedCoMs(), BodyIndex))
 		{
-			// Return the Selection objects CoM position when manipulating as that is the one we're actually updating with the 
-			// manipulator widget (as updating the CoM in the physics body proper is complicated).
-			return SelectedCoM->CoMPosition;
+			if (const FVector* const ManipulatedCoMPosition = FindManipulatedBodyCoMPosition(BodyIndex))
+			{
+				// Return the Selection objects CoM position when manipulating as that is the one we're actually updating with the 
+				// manipulator widget (as updating the CoM in the physics body proper is complicated).
+				return *ManipulatedCoMPosition;
+			}
 		}
 	}
 
@@ -480,36 +455,39 @@ FVector FPhysicsAssetEditorSharedData::CalculateCoMNudgeForWorldSpacePosition(co
 	return CalculatedCoMOffset;
 }
 
-void FPhysicsAssetEditorSharedData::RecordSelectedBodiesCoM()
+void FPhysicsAssetEditorSharedData::RecordSelectedCoM()
 {
 	if (EditorSkelComp)
 	{
-		for(FSelection& Body : SelectedBodies)
+		ManipulatedBodyCoMPositionMap.Reset();
+
+		for (const FSelection& SelectedObject : SelectedObjects->UniqueSelectedElementsOfType(FSelection::Body | FSelection::Primitive | FSelection::CenterOfMass))
 		{
-			Body.CoMPosition = EditorSkelComp->Bodies[Body.Index]->GetCOMPosition();
+			ManipulatedBodyCoMPositionMap.FindOrAdd(SelectedObject.GetIndex()) = EditorSkelComp->Bodies[SelectedObject.GetIndex()]->GetCOMPosition();
 		}
 	}
 }
 
-void FPhysicsAssetEditorSharedData::UpdateSelectedBodiesCoM()
+void FPhysicsAssetEditorSharedData::PostManipulationUpdateCoM()
 {
-	TSet<int32> EncounteredBodyIndexs;
-
-	// Update CoM nudge to compensate for any change in body transform on any axis that is fixed in component space.
-	for (int32 i = 0; i < SelectedBodies.Num(); ++i)
+	// Update CoM nudge to compensate for any change in body transform on any axis that is fixed in component space.		
+	for (const FSelection& SelectedObject : SelectedObjects->UniqueSelectedElementsOfType(FSelection::Body | FSelection::Primitive | FSelection::CenterOfMass))
 	{
-		FPhysicsAssetEditorSharedData::FSelection& SelectedObject = SelectedBodies[i];
-		const int32 BodyIndex = SelectedObject.Index;
-		FBodyInstance* const EditorBodyInstance = EditorSkelComp->Bodies[BodyIndex];
+		const int32 BodyIndex = SelectedObject.GetIndex();
+		const FBodyInstance* const EditorBodyInstance = EditorSkelComp->Bodies[BodyIndex];
+		const FVector* const ManipulationCoMPosition = FindManipulatedBodyCoMPosition(BodyIndex);
 
-		if (IsCoMAxisFixedInComponentSpace(BodyIndex, EAxis::X) || IsCoMAxisFixedInComponentSpace(BodyIndex, EAxis::Y) || IsCoMAxisFixedInComponentSpace(BodyIndex, EAxis::Z))
+		if (ManipulationCoMPosition) // Expect to find a valid cached CoM position for any selected CoM marker or primitive undergoing manipulation.
 		{
-			// Ensure we only update the CoM of each physics body once (SelectedBodies actually contains a list of selected primitives).
-			if (!EncounteredBodyIndexs.Contains(BodyIndex)) 
+			if (SelectedObject.HasType(FSelection::CenterOfMass)) // Directly selected CoM markers have priority over their owning bodies for determining CoM manipulation behavior.
+			{
+				const FVector CalculatedCoMOffset = CalculateCoMNudgeForWorldSpacePosition(BodyIndex, *ManipulationCoMPosition);
+				PhysicsAsset->SkeletalBodySetups[BodyIndex]->DefaultInstance.COMNudge = CalculatedCoMOffset;
+			}
+			else if (IsCoMAxisFixedInComponentSpace(BodyIndex, EAxis::X) || IsCoMAxisFixedInComponentSpace(BodyIndex, EAxis::Y) || IsCoMAxisFixedInComponentSpace(BodyIndex, EAxis::Z))
 			{
-				EncounteredBodyIndexs.Add(BodyIndex);
 				const FVector CoMOffset = EditorBodyInstance->COMNudge;
-				FVector	CalculatedCoMOffset = CalculateCoMNudgeForWorldSpacePosition(BodyIndex, SelectedObject.CoMPosition);
+				FVector	CalculatedCoMOffset = CalculateCoMNudgeForWorldSpacePosition(BodyIndex, *ManipulationCoMPosition);
 
 				// Only apply lock to the specified Axis in bone space.
 				if (!IsCoMAxisFixedInComponentSpace(BodyIndex, EAxis::X)) { CalculatedCoMOffset.X = CoMOffset.X; }
@@ -522,6 +500,17 @@ void FPhysicsAssetEditorSharedData::UpdateSelectedBodiesCoM()
 	}
 }
 
+void FPhysicsAssetEditorSharedData::UpdateCoM()
+{
+	if (bShouldUpdatedSelectedCoMs) // < This calculation must be delayed by a frame s.t. changes to the physics state have been propagated to the physics bodies.
+	{
+		PostManipulationUpdateCoM();
+		RefreshPhysicsAssetChange(PhysicsAsset, false);
+		ManipulatedBodyCoMPositionMap.Reset();
+		bShouldUpdatedSelectedCoMs = false;
+	}
+}
+
 bool FPhysicsAssetEditorSharedData::ParseClipboard(UPhysicsAsset*& OutAsset, FString& OutObjectType, UObject*& OutObject)
 {
 	FString ClipboardContent;
@@ -596,9 +585,9 @@ void FPhysicsAssetEditorSharedData::Mirror()
 	{
 		// Build list of all bodies and constraints to be mirrored
 		TArray<FMirrorInfo> MirrorInfos;
-		MirrorInfos.Reserve(SelectedBodies.Num() + SelectedConstraints.Num());
+		MirrorInfos.Reserve(UniqueSelectionReferencingBodies().Num() + SelectedConstraints().Num());
 
-		for (const FSelection& Selection : SelectedBodies)
+		for (const FSelection& Selection : UniqueSelectionReferencingBodies())
 		{
 			MirrorInfos.AddDefaulted();
 			FMirrorInfo & MirrorInfo = MirrorInfos[MirrorInfos.Num() - 1];
@@ -620,7 +609,7 @@ void FPhysicsAssetEditorSharedData::Mirror()
 			}
 		}
 
-		for (const FSelection& Selection : SelectedConstraints)
+		for (const FSelection& Selection : SelectedConstraints())
 		{
 			MirrorInfos.AddDefaulted();
 			FMirrorInfo & MirrorInfo = MirrorInfos[MirrorInfos.Num() - 1];
@@ -804,28 +793,22 @@ void FPhysicsAssetEditorSharedData::HitBone(int32 BodyIndex, EAggCollisionShape:
 {
 	if (!bRunningSimulation)
 	{
-		FPhysicsAssetEditorSharedData::FSelection Selection(BodyIndex, PrimType, PrimIndex);
-
-		if (EditorSkelComp && EditorSkelComp->Bodies.IsValidIndex(BodyIndex))
-		{
-			Selection.CoMPosition = EditorSkelComp->Bodies[BodyIndex]->GetCOMPosition();
-		}
+		FPhysicsAssetEditorSharedData::FSelection Selection = MakePrimitiveSelection(BodyIndex, PrimType, PrimIndex);
 
 		if(bGroupSelect)
 		{
 			if(IsBodySelected(Selection))
 			{
-				SetSelectedBody(Selection, false);
+				ModifySelectedPrimitives(Selection, false);
 			}
 			else
 			{
-				SetSelectedBody(Selection, true);
+				ModifySelectedPrimitives(Selection, true);
 			}
 		}
 		else
 		{
-			ClearSelectedBody();
-			SetSelectedBody(Selection, true);
+			SetSelectedPrimitives(Selection);
 		}
 	}
 }
@@ -834,28 +817,22 @@ void FPhysicsAssetEditorSharedData::HitCoM(const int32 BodyIndex, const bool bGr
 {
 	if (!bRunningSimulation)
 	{
-		FPhysicsAssetEditorSharedData::FSelection Selection(BodyIndex, EAggCollisionShape::Unknown, INDEX_NONE);
-
-		if (EditorSkelComp && EditorSkelComp->Bodies.IsValidIndex(BodyIndex))
-		{
-			Selection.CoMPosition = EditorSkelComp->Bodies[BodyIndex]->GetCOMPosition();
-		}
+		const FPhysicsAssetEditorSharedData::FSelection Selection = MakeCoMSelection(BodyIndex);
 
 		if (bGroupSelect)
 		{
 			if (IsCoMSelected(BodyIndex))
 			{
-				SetSelectedCoM(Selection, false);
+				ModifySelectedCoMs(Selection, false);
 			}
 			else
 			{
-				SetSelectedCoM(Selection, true);
+				ModifySelectedCoMs(Selection, true);
 			}
 		}
 		else
 		{
-			ClearSelectedCoMs();
-			SetSelectedCoM(Selection, true);
+			SetSelectedCoMs(Selection);
 		}
 	}
 }
@@ -868,17 +845,17 @@ void FPhysicsAssetEditorSharedData::HitConstraint(int32 ConstraintIndex, bool bG
 		{
 			if(IsConstraintSelected(ConstraintIndex))
 			{
-				SetSelectedConstraint(ConstraintIndex, false);
+				ModifySelectedConstraints(ConstraintIndex, false);
 			}
 			else
 			{
-				SetSelectedConstraint(ConstraintIndex, true);
+				ModifySelectedConstraints(ConstraintIndex, true);
 			}
 		}
 		else
 		{
 			ClearSelectedConstraints();
-			SetSelectedConstraint(ConstraintIndex, true);
+			ModifySelectedConstraints(ConstraintIndex, true);
 		}
 	}
 }
@@ -915,30 +892,25 @@ void FPhysicsAssetEditorSharedData::RefreshPhysicsAssetChange(const UPhysicsAsse
 	}
 }
 
-void FPhysicsAssetEditorSharedData::SetSelectedBodyAnyPrimitive(int32 BodyIndex, bool bSelected)
-{
-	SetSelectedBodiesAnyPrimitive({ BodyIndex }, bSelected);
-}
-
-void FPhysicsAssetEditorSharedData::SetSelectedBodiesAnyPrimitive(const TArray<int32>& BodiesIndices, const bool bSelected)
+void FPhysicsAssetEditorSharedData::SetSelectedBodiesAllPrimitive(const TArray<int32>& BodiesIndices, bool bSelected)
 {
 	SetSelectedBodiesPrimitives(BodiesIndices, bSelected, [](const TArray<FSelection>& CurrentSelection, const int32 BodyIndex, const FKShapeElem& Primitive)
 	{
-		// Select the first primitive in the supplied body. 
-		return !CurrentSelection.ContainsByPredicate([BodyIndex](const FSelection& InSelection) {return InSelection.Index == BodyIndex; }); // Predicate returns true if the supplied body index is not already in the current selection.
+		// Select all primitives
+		return true;
 	});
 }
 
-void FPhysicsAssetEditorSharedData::SetSelectedBodiesAllPrimitive(const TArray<int32>& BodiesIndices, bool bSelected)
+void FPhysicsAssetEditorSharedData::SetSelectedBodiesPrimitivesWithCollisionType(const TArray<int32>& BodiesIndices, const ECollisionEnabled::Type CollisionType, bool bSelected)
 {
-	SetSelectedBodiesPrimitives(BodiesIndices, bSelected, [](const TArray<FSelection>& CurrentSelection, const int32 BodyIndex, const FKShapeElem& Primitive)
+	SetSelectedBodiesPrimitives(BodiesIndices, bSelected, [CollisionType](const TArray<FSelection>& CurrentSelection, const int32 BodyIndex, const FKShapeElem& Primitive)
 	{
-		// Select all primitives
-		return true;
+		// Select primitives which match the collision type
+		return Primitive.GetCollisionEnabled() == CollisionType;
 	});
 }
 
-void FPhysicsAssetEditorSharedData::SetSelectedBodiesFromDirectlInput(const TArray<int32>& BodiesIndices, const bool bSelected)
+void FPhysicsAssetEditorSharedData::SetSelectedBodiesPrimitives(const TArray<int32>& BodiesIndices, bool bSelected, const TFunction<bool(const TArray<FSelection>&, const int32 BodyIndex, const FKShapeElem&)>& Predicate)
 {
 	if (BodiesIndices.Num() == 0)
 	{
@@ -951,53 +923,135 @@ void FPhysicsAssetEditorSharedData::SetSelectedBodiesFromDirectlInput(const TArr
 		return;
 	}
 
-	// Create a selection for each body that references the first primitive in that body. (ordering of primitives is determined by geometry type then by order if instantiation).
-	TArray<FSelection> SelectionCollection = CreateBodyPrimitivesSelection(PhysicsAsset, BodiesIndices, [this](const TArray<FSelection>& CurrentSelection, const int32 BodyIndex, const FKShapeElem& Primitive) 
-		{ 
-			// Select the first primitive in the supplied body. 
-			return !SelectionContainsIndex(CurrentSelection, BodyIndex); // Predicate returns true if the supplied body index is not already in the current selection.
-		});
+	ModifySelectedBodies(CreateBodyPrimitivesSelection(PhysicsAsset, BodiesIndices, Predicate), bSelected);
+}
 
-	for (FSelection& Selection : SelectionCollection)
-	{
-		Selection.bIsDirectBodySelection = true;
-	}
+FPhysicsAssetEditorSharedData::SelectionFilterRange FPhysicsAssetEditorSharedData::SelectedBodies() const
+{ 
+	return SelectedObjects->SelectedElementsOfType(FSelection::Body);
+}
 
-	SetSelectedBodies(SelectionCollection, bSelected);
+FPhysicsAssetEditorSharedData::SelectionFilterRange FPhysicsAssetEditorSharedData::SelectedCoMs() const
+{ 
+	return SelectedObjects->SelectedElementsOfType(FSelection::CenterOfMass);
 }
 
-void FPhysicsAssetEditorSharedData::SetSelectedBodiesPrimitivesWithCollisionType(const TArray<int32>& BodiesIndices, const ECollisionEnabled::Type CollisionType, bool bSelected)
+FPhysicsAssetEditorSharedData::SelectionFilterRange FPhysicsAssetEditorSharedData::SelectedConstraints() const
 {
-	SetSelectedBodiesPrimitives(BodiesIndices, bSelected, [CollisionType](const TArray<FSelection>& CurrentSelection, const int32 BodyIndex, const FKShapeElem& Primitive)
-	{
-		// Select primitives which match the collision type
-		return Primitive.GetCollisionEnabled() == CollisionType;
-	});
+	return SelectedObjects->SelectedElementsOfType(FSelection::Constraint);
 }
 
-void FPhysicsAssetEditorSharedData::SetSelectedBodiesPrimitives(const TArray<int32>& BodiesIndices, bool bSelected, const TFunction<bool(const TArray<FSelection>&, const int32 BodyIndex, const FKShapeElem&)>& Predicate)
+FPhysicsAssetEditorSharedData::SelectionFilterRange FPhysicsAssetEditorSharedData::SelectedPrimitives() const
 {
-	if (BodiesIndices.Num() == 0)
-	{
-		return;
-	}
+	return SelectedObjects->SelectedElementsOfType(FSelection::Primitive);
+}
 
-	if (BodiesIndices.Num() == 1 && BodiesIndices[0] == INDEX_NONE)
+FPhysicsAssetEditorSharedData::SelectionFilterRange FPhysicsAssetEditorSharedData::SelectedBodiesAndPrimitives() const
+{
+	return SelectedObjects->SelectedElementsOfType(FSelection::Body | FSelection::Primitive);
+}
+
+
+
+FPhysicsAssetEditorSharedData::SelectionUniqueRange FPhysicsAssetEditorSharedData::UniqueSelectionReferencingBodies() const
+{
+	return SelectedObjects->UniqueSelectedElementsOfType(FSelection::Body | FSelection::Primitive);
+}
+
+const UPhysicsAssetEditorSelection* FPhysicsAssetEditorSharedData::GetSelectedObjects() const
+{
+	return SelectedObjects;
+}
+
+const FPhysicsAssetEditorSharedData::FSelection* FPhysicsAssetEditorSharedData::GetSelectedBodyOrPrimitive() const
+{
+	return SelectedObjects->GetLastSelectedOfType(FSelection::Body | FSelection::Primitive);
+}
+
+const FPhysicsAssetEditorSharedData::FSelection* FPhysicsAssetEditorSharedData::GetSelectedBody() const
+{
+	return SelectedObjects->GetLastSelectedOfType(FSelection::Body);
+}
+
+const FPhysicsAssetEditorSharedData::FSelection* FPhysicsAssetEditorSharedData::GetSelectedCoM() const
+{
+	return SelectedObjects->GetLastSelectedOfType(FSelection::CenterOfMass);
+}
+
+const FPhysicsAssetEditorSharedData::FSelection* FPhysicsAssetEditorSharedData::GetSelectedConstraint() const
+{
+	return SelectedObjects->GetLastSelectedOfType(FSelection::Constraint);
+}
+
+const FPhysicsAssetEditorSharedData::FSelection* FPhysicsAssetEditorSharedData::GetSelectedPrimitive() const
+{
+	return SelectedObjects->GetLastSelectedOfType(FSelection::Primitive);
+}
+
+void FPhysicsAssetEditorSharedData::SetGroupSelectionActive(const bool bIsActive)
+{
+	bIsGroupSelectionActive = bIsActive;
+}
+
+bool FPhysicsAssetEditorSharedData::IsGroupSelectionActive() const
+{
+	return bIsGroupSelectionActive;
+}
+
+void FPhysicsAssetEditorSharedData::ModifySelected(const TArray<FSelection>& InSelectedElements, const bool bSelected)
+{
+	ModifySelectionInternal([this, InSelectedElements, bSelected]() -> bool { return SelectedObjects->ModifySelected(InSelectedElements, bSelected); });
+}
+
+void FPhysicsAssetEditorSharedData::SetSelected(const TArray<FSelection>& InSelectedElements)
+{
+	ModifySelectionInternal([this, InSelectedElements]() -> bool { return SelectedObjects->SetSelected(InSelectedElements); });
+}
+
+bool FPhysicsAssetEditorSharedData::IsSelected(const FSelection& InSelection) const
+{
+	return SelectedObjects->SelectedElements().Contains(InSelection);
+}
+
+void FPhysicsAssetEditorSharedData::ClearSelected()
+{
+	SelectedObjects->ClearSelection();
+
+	BroadcastSelectionChanged();
+}
+
+void FPhysicsAssetEditorSharedData::ClearSelectedPrimitives()
+{
+	if (InsideSelChange)
 	{
-		ClearSelectedBody();
 		return;
 	}
 
-	SetSelectedBodies(CreateBodyPrimitivesSelection(PhysicsAsset, BodiesIndices, Predicate), bSelected);
+	ClearSelected();
+
+	++InsideSelChange;
+	BroadcastPreviewChanged();
+	--InsideSelChange;
 }
 
-void FPhysicsAssetEditorSharedData::ClearSelected()
+void FPhysicsAssetEditorSharedData::ModifySelectedPrimitives(const FSelection& InSelectedElement, const bool bSelected)
 {
-	SelectedBodies.Empty();
-	SelectedCoMs.Empty();
-	SelectedConstraints.Empty();
+	ModifySelected(TArray<FSelection>{ InSelectedElement }, bSelected);
+}
 
-	BroadcastSelectionChanged();
+void FPhysicsAssetEditorSharedData::ModifySelectedPrimitives(const TArray<FSelection>& InSelectedElements, const bool bSelected)
+{
+	ModifySelected(InSelectedElements, bSelected);
+}
+
+void FPhysicsAssetEditorSharedData::SetSelectedPrimitives(const FSelection& InSelectedElement)
+{
+	SetSelected(TArray<FSelection>{ InSelectedElement });
+}
+
+void FPhysicsAssetEditorSharedData::SetSelectedPrimitives(const TArray<FSelection>& InSelectedElements)
+{
+	SetSelected(InSelectedElements); // TODO - should only clear and set selected primitives, not bodies etc
 }
 
 void FPhysicsAssetEditorSharedData::ClearSelectedCoMs()
@@ -1014,26 +1068,29 @@ void FPhysicsAssetEditorSharedData::ClearSelectedCoMs()
 	--InsideSelChange;
 }
 
-void FPhysicsAssetEditorSharedData::SetSelectedCoM(const FSelection& InSelectedElement, const bool bSelected)
+void FPhysicsAssetEditorSharedData::ModifySelectedCoMs(const FSelection& InSelectedElement, const bool bSelected)
 {
-	SetSelectedCoMs({ InSelectedElement }, bSelected);
+	ModifySelected(TArray<FSelection>{ InSelectedElement }, bSelected);
 }
 
-void FPhysicsAssetEditorSharedData::SetSelectedCoMs(const TArray<FSelection>& InSelectedElements, const bool bSelected)
+void FPhysicsAssetEditorSharedData::ModifySelectedCoMs(const TArray<FSelection>& InSelectedElements, const bool bSelected)
 {
-	if (!InsideSelChange && SetSelected(SelectedCoMs, InSelectedElements, bSelected))
-	{
-		BroadcastSelectionChanged();
+	ModifySelected(InSelectedElements, bSelected);
+}
 
-		++InsideSelChange;
-		BroadcastPreviewChanged();
-		--InsideSelChange;
-	}	
+void FPhysicsAssetEditorSharedData::SetSelectedCoMs(const FSelection& InSelectedElement)
+{
+	SetSelected(TArray<FSelection>{ InSelectedElement });
+}
+
+void FPhysicsAssetEditorSharedData::SetSelectedCoMs(const TArray<FSelection>& InSelectedElements)
+{
+	SetSelected(InSelectedElements);
 }
 
 bool FPhysicsAssetEditorSharedData::IsCoMSelected(const int32 BodyIndex) const
 {
-	return SelectedCoMs.ContainsByPredicate([BodyIndex](const FSelection& Element) { return Element.Index == BodyIndex; });
+	return SelectionContainsIndex(SelectedCoMs(), BodyIndex);
 }
 
 void FPhysicsAssetEditorSharedData::ClearSelectedBody()
@@ -1041,56 +1098,65 @@ void FPhysicsAssetEditorSharedData::ClearSelectedBody()
 	ClearSelected();
 }
 
-void FPhysicsAssetEditorSharedData::SetSelectedBody(const FSelection& Body, bool bSelected)
+void FPhysicsAssetEditorSharedData::ModifySelectedBodies(const FSelection& Body, bool bSelected)
 {
-	SetSelectedBodies({ Body }, bSelected);
+	ModifySelected(TArray<FSelection>{ Body }, bSelected);
 }
 
-void FPhysicsAssetEditorSharedData::SetSelectedBodies(const TArray<FSelection>& Bodies, bool bSelected)
+void FPhysicsAssetEditorSharedData::ModifySelectedBodies(const TArray<FSelection>& InSelectedElements, bool bSelected)
 {
-	if (!InsideSelChange && SetSelected(SelectedBodies, Bodies, bSelected))
-	{
-		BroadcastSelectionChanged();
+	ModifySelected(InSelectedElements, bSelected);
+}
 
-		if (!GetSelectedBody())
-		{
-			return;
-		}
+void FPhysicsAssetEditorSharedData::SetSelectedBodies(const FSelection& InSelectedElement)
+{
+	SetSelected(TArray<FSelection>{ InSelectedElement });
+}
 
-		UpdateNoCollisionBodies();
-		++InsideSelChange;
-		BroadcastPreviewChanged();
-		--InsideSelChange;
-	}
+void FPhysicsAssetEditorSharedData::SetSelectedBodies(const TArray<FSelection>& InSelectedElements)
+{
+	SetSelected(InSelectedElements);
 }
 
-void FPhysicsAssetEditorSharedData::SetGroupSelectionActive(const bool bIsActive)
+void FPhysicsAssetEditorSharedData::ModifySelectedBodies(const int32 BodyIndex, const bool bSelected)
 {
-	bIsGroupSelectionActive = bIsActive;
+	ModifySelected(TArray<FSelection>{ MakeBodySelection(PhysicsAsset, BodyIndex) }, bSelected);
 }
 
-bool FPhysicsAssetEditorSharedData::IsGroupSelectionActive() const
+void FPhysicsAssetEditorSharedData::ModifySelectedBodies(const TArray<int32>& BodiesIndices, const bool bSelected)
 {
-	return bIsGroupSelectionActive;
+	ModifySelected(MakeBodySelection(PhysicsAsset, BodiesIndices), bSelected);
+}
+
+void FPhysicsAssetEditorSharedData::SetSelectedBodies(const int32 BodyIndex)
+{
+	SetSelected(TArray<FSelection>{ MakeBodySelection(PhysicsAsset, BodyIndex) });
+}
+
+void FPhysicsAssetEditorSharedData::SetSelectedBodies(const TArray<int32>& BodiesIndices)
+{
+	SetSelected(MakeBodySelection(PhysicsAsset, BodiesIndices));
 }
 
 bool FPhysicsAssetEditorSharedData::IsBodySelected(const FSelection& Body) const
 {
-	return SelectedBodies.Contains(Body);
+	return Body.HasType(FPhysicsAssetEditorSelectedElement::Body) && SelectedObjects->SelectedElements().Contains(Body); // TODO - should this be implemented with the following fn ?
 }
 
 bool FPhysicsAssetEditorSharedData::IsBodySelected(const int32 BodyIndex) const
 {
-	return SelectedBodies.ContainsByPredicate([BodyIndex](const FSelection& SelectedBody) { return SelectedBody.Index == BodyIndex; });
+	return Algo::FindByPredicate(SelectedObjects->SelectedElements(), [BodyIndex](const FSelection& Element) 
+		{ 
+			return Element.HasType(FPhysicsAssetEditorSelectedElement::Body) && (Element.Index == BodyIndex); 
+		}) != nullptr;
 }
 
 void FPhysicsAssetEditorSharedData::ToggleSelectionType(bool bIgnoreUserConstraints)
 {
 	TSet<int32> NewSelectedBodies; 
-	for (const FSelection& Selection : SelectedConstraints)
+	for (const FSelection& Selection : SelectedConstraints())
 	{
-		UPhysicsConstraintTemplate* ConstraintTemplate = PhysicsAsset->ConstraintSetup[Selection.Index];
-		FConstraintInstance & DefaultInstance = ConstraintTemplate->DefaultInstance;
+		const UPhysicsConstraintTemplate* const ConstraintTemplate = PhysicsAsset->ConstraintSetup[Selection.Index];
 
 		for (int32 BodyIdx = 0; BodyIdx < PhysicsAsset->SkeletalBodySetups.Num(); ++BodyIdx)
 		{
@@ -1107,8 +1173,8 @@ void FPhysicsAssetEditorSharedData::ToggleSelectionType(bool bIgnoreUserConstrai
 		}
 	}
 
-	TSet<int32> NewSelectedConstraints; //We could have multiple shapes selected which would cause us to add and remove the same constraint.
-	for (const FSelection& Selection : SelectedBodies)
+	TSet<int32> NewSelectedConstraints; // Use a set here because we could have multiple shapes selected which would cause us to add and remove the same constraint.
+	for (const FSelection& Selection : UniqueSelectionReferencingBodies())
 	{
 		UBodySetup* BodySetup = PhysicsAsset->SkeletalBodySetups[Selection.Index];
 		for(int32 ConstraintIdx = 0; ConstraintIdx < PhysicsAsset->ConstraintSetup.Num(); ++ConstraintIdx)
@@ -1134,7 +1200,7 @@ void FPhysicsAssetEditorSharedData::ToggleSelectionType(bool bIgnoreUserConstrai
 	ClearSelectedConstraints();
 
 	SetSelectedBodiesAllPrimitive(NewSelectedBodies.Array(), true);
-	SetSelectedConstraints(NewSelectedConstraints.Array(), true);
+	ModifySelectedConstraints(NewSelectedConstraints.Array(), true);
 }
 
 void FPhysicsAssetEditorSharedData::ToggleShowSelected()
@@ -1142,7 +1208,7 @@ void FPhysicsAssetEditorSharedData::ToggleShowSelected()
 	bool bAllSelectedVisible = true;
 	if (bAllSelectedVisible)
 	{
-		for (const FSelection& Selection : SelectedConstraints)
+		for (const FSelection& Selection : SelectedConstraints())
 		{
 			if (IsConstraintHidden(Selection.Index))
 			{
@@ -1153,7 +1219,7 @@ void FPhysicsAssetEditorSharedData::ToggleShowSelected()
 	}
 	if (bAllSelectedVisible)
 	{
-		for (const FSelection& Selection : SelectedBodies)
+		for (const FSelection& Selection : UniqueSelectionReferencingBodies())
 		{
 			if (IsBodyHidden(Selection.Index))
 			{
@@ -1179,7 +1245,7 @@ void FPhysicsAssetEditorSharedData::ToggleShowOnlySelected()
 	bool bAllSelectedVisible = true;
 	if (bAllSelectedVisible)
 	{
-		for (const FSelection& Selection : SelectedConstraints)
+		for (const FSelection& Selection : SelectedConstraints())
 		{
 			if (IsConstraintHidden(Selection.Index))
 			{
@@ -1190,7 +1256,7 @@ void FPhysicsAssetEditorSharedData::ToggleShowOnlySelected()
 	}
 	if (bAllSelectedVisible)
 	{
-		for (const FSelection& Selection : SelectedBodies)
+		for (const FSelection& Selection : UniqueSelectionReferencingBodies())
 		{
 			if (IsBodyHidden(Selection.Index))
 			{
@@ -1205,7 +1271,7 @@ void FPhysicsAssetEditorSharedData::ToggleShowOnlySelected()
 		for (int32 ConstraintIndex = 0; ConstraintIndex < PhysicsAsset->ConstraintSetup.Num(); ++ConstraintIndex)
 		{
 			// Look at unselected constraints
-			if (!SelectedConstraints.ContainsByPredicate([ConstraintIndex](FSelection& V) { return V.Index == ConstraintIndex; } ))
+			if (!SelectionContainsIndex(SelectedConstraints(), ConstraintIndex))
 			{
 				// Is it hidden?
 				if (!IsConstraintHidden(ConstraintIndex))
@@ -1221,7 +1287,7 @@ void FPhysicsAssetEditorSharedData::ToggleShowOnlySelected()
 		for (int32 BodyIndex = 0; BodyIndex < PhysicsAsset->SkeletalBodySetups.Num(); ++BodyIndex)
 		{
 			// Look at unselected bodies
-			if (!SelectedBodies.ContainsByPredicate([BodyIndex](FSelection& V) { return V.Index == BodyIndex; }))
+			if (!IsBodySelected(BodyIndex))
 			{
 				// Is it hidden?
 				if (!IsBodyHidden(BodyIndex))
@@ -1328,11 +1394,11 @@ void FPhysicsAssetEditorSharedData::HideAll()
 
 void FPhysicsAssetEditorSharedData::ShowSelected()
 {
-	for (const FSelection& Selection : SelectedConstraints)
+	for (const FSelection& Selection : SelectedConstraints())
 	{
 		ShowConstraint(Selection.Index);
 	}
-	for (const FSelection& Selection : SelectedBodies)
+	for (const FSelection& Selection : UniqueSelectionReferencingBodies())
 	{
 		ShowBody(Selection.Index);
 	}
@@ -1340,11 +1406,11 @@ void FPhysicsAssetEditorSharedData::ShowSelected()
 
 void FPhysicsAssetEditorSharedData::HideSelected()
 {
-	for (const FSelection& Selection : SelectedConstraints)
+	for (const FSelection& Selection : SelectedConstraints())
 	{
 		HideConstraint(Selection.Index);
 	}
-	for (const FSelection& Selection : SelectedBodies)
+	for (const FSelection& Selection : UniqueSelectionReferencingBodies())
 	{
 		HideBody(Selection.Index);
 	}
@@ -1371,7 +1437,7 @@ void FPhysicsAssetEditorSharedData::ToggleShowOnlyColliding()
 	FPhysicsAssetRenderSettings* const PhysicsAssetRenderSettings = GetRenderSettings();
 	
 	// only works if one only body is selected
-	if (!bIsShowingColliding && PhysicsAssetRenderSettings && (SelectedBodies.Num() == 1))
+	if (!bIsShowingColliding && PhysicsAssetRenderSettings && (UniqueSelectionReferencingBodies().Num() == 1))
 	{	
 		// NoCollisionBodies already contains the non colliding bodies from the one selection
 		PhysicsAssetRenderSettings->SetHiddenBodies(NoCollisionBodies);
@@ -1400,13 +1466,13 @@ void FPhysicsAssetEditorSharedData::ToggleShowOnlyConstrained()
 	HideAllBodies();
 
 	// add  the current selection of bodies
-	for (const FSelection& SelectedBody : SelectedBodies)
+	for (const FSelection& Selection : UniqueSelectionReferencingBodies())
 	{
-		ShowBody(SelectedBody.Index);
+		ShowBody(Selection.Index);
 	}
 
 	// collect connected bodies from the selected constraints
-	for (const FSelection& Selection : SelectedConstraints)
+	for (const FSelection& Selection : SelectedConstraints())
 	{
 		UPhysicsConstraintTemplate* ConstraintTemplate = PhysicsAsset->ConstraintSetup[Selection.Index];
 		FConstraintInstance& DefaultInstance = ConstraintTemplate->DefaultInstance;
@@ -1425,7 +1491,7 @@ void FPhysicsAssetEditorSharedData::ToggleShowOnlyConstrained()
 	}
 
 	// collect connected bodies from the selected bodies
-	for (const FSelection& Selection : SelectedBodies)
+	for (const FSelection& Selection : UniqueSelectionReferencingBodies())
 	{
 		UBodySetup* BodySetup = PhysicsAsset->SkeletalBodySetups[Selection.Index];
 		for (int32 ConstraintIdx = 0; ConstraintIdx < PhysicsAsset->ConstraintSetup.Num(); ++ConstraintIdx)
@@ -1468,16 +1534,16 @@ void FPhysicsAssetEditorSharedData::UpdateNoCollisionBodies()
 		{
 			NoCollisionBodies.Add(i);
 		}
-		else if (GetSelectedBody() && i != GetSelectedBody()->Index)
+		else if (GetSelectedBodyOrPrimitive() && i != GetSelectedBodyOrPrimitive()->Index)
 		{
-			if (!ensure(PhysicsAsset->SkeletalBodySetups[GetSelectedBody()->Index]))
+			if (!ensure(PhysicsAsset->SkeletalBodySetups[GetSelectedBodyOrPrimitive()->Index]))
 			{
 				continue;
 			}
 			// Add this body if it has disabled collision with selected.
-			FRigidBodyIndexPair Key(i, GetSelectedBody()->Index);
+			FRigidBodyIndexPair Key(i, GetSelectedBodyOrPrimitive()->Index);
 
-			if (PhysicsAsset->SkeletalBodySetups[GetSelectedBody()->Index]->DefaultInstance.GetCollisionEnabled() == ECollisionEnabled::NoCollision ||
+			if (PhysicsAsset->SkeletalBodySetups[GetSelectedBodyOrPrimitive()->Index]->DefaultInstance.GetCollisionEnabled() == ECollisionEnabled::NoCollision ||
 				PhysicsAsset->CollisionDisableTable.Find(Key))
 			{
 				NoCollisionBodies.Add(i);
@@ -1500,128 +1566,114 @@ void FPhysicsAssetEditorSharedData::ClearSelectedConstraints()
 	--InsideSelChange;
 }
 
-void FPhysicsAssetEditorSharedData::SetSelectedConstraint(int32 ConstraintIndex, bool bSelected)
+void FPhysicsAssetEditorSharedData::ModifySelectedConstraints(const int32 ConstraintIndex, const bool bSelected)
 {
-	SetSelectedConstraints({ ConstraintIndex }, bSelected);
+	ModifySelectedConstraints(TArray<int32>{ ConstraintIndex }, bSelected);
 }
 
-void FPhysicsAssetEditorSharedData::SetSelectedConstraints(const TArray<int32> ConstraintsIndices, bool bSelected)
+void FPhysicsAssetEditorSharedData::ModifySelectedConstraints(const TArray<int32>& ConstraintsIndices, const bool bSelected)
 {
-	if (!InsideSelChange && SetSelected(SelectedConstraints, ConstraintsIndices, bSelected))
-	{
-		BroadcastSelectionChanged();
+	ModifySelected(MakeConstraintSelection(ConstraintsIndices), bSelected);
+}
 
-		++InsideSelChange;
-		BroadcastPreviewChanged();
-		--InsideSelChange;
-	}
+void FPhysicsAssetEditorSharedData::SetSelectedConstraints(const TArray<int32>& ConstraintsIndices)
+{
+	SetSelected(MakeConstraintSelection(ConstraintsIndices));
 }
 
-bool FPhysicsAssetEditorSharedData::IsConstraintSelected(int32 ConstraintIndex) const
+bool FPhysicsAssetEditorSharedData::IsConstraintSelected(const int32 ConstraintIndex) const
 {
-	FSelection Constraint(ConstraintIndex, EAggCollisionShape::Unknown, INDEX_NONE);
-	return SelectedConstraints.Contains(Constraint);
+	return SelectedObjects->SelectedElements().Contains(MakeConstraintSelection(ConstraintIndex));
 }
 
-void FPhysicsAssetEditorSharedData::SetCollisionBetweenSelected(bool bEnableCollision)
+void FPhysicsAssetEditorSharedData::SetCollisionBetweenSelected(const bool bEnableCollision)
 {
-	if (bRunningSimulation || SelectedBodies.Num() == 0)
+	if (bRunningSimulation || UniqueSelectionReferencingBodies().IsEmpty())
 	{
 		return;
 	}
 
 	PhysicsAsset->Modify();
 
-	for(int32 i=0; i<SelectedBodies.Num(); ++i)
-	{
-		for(int32 j=i+1; j<SelectedBodies.Num(); ++j)
+	ForEachUniquePair(UniqueSelectionReferencingBodies(), [PhysicsAsset = PhysicsAsset, bEnableCollision](const FSelection& Lhs, const FSelection& Rhs)
 		{
-			if(bEnableCollision)
+			if (bEnableCollision)
 			{
-				PhysicsAsset->EnableCollision(SelectedBodies[i].Index, SelectedBodies[j].Index);
-			}else
+				PhysicsAsset->EnableCollision(Lhs.Index, Rhs.Index);
+			}
+			else
 			{
-				PhysicsAsset->DisableCollision(SelectedBodies[i].Index, SelectedBodies[j].Index);
+				PhysicsAsset->DisableCollision(Lhs.Index, Rhs.Index);
 			}
+		});
 
-		}
-	}
-
-
-	UpdateNoCollisionBodies();
-	
+	UpdateNoCollisionBodies();	
 	RefreshPhysicsAssetChange(PhysicsAsset);
-
 	BroadcastPreviewChanged();
 }
 
 bool FPhysicsAssetEditorSharedData::CanSetCollisionBetweenSelected(bool bEnableCollision) const
 {
-	if (bRunningSimulation || SelectedBodies.Num() == 0)
+	if (bRunningSimulation || UniqueSelectionReferencingBodies().IsEmpty())
 	{
 		return false;
 	}
 
-	for(int32 i=0; i<SelectedBodies.Num(); ++i)
-	{
-		for(int32 j=i+1; j<SelectedBodies.Num(); ++j)
+	bool bResult = false;
+
+	ForEachUniquePair(UniqueSelectionReferencingBodies(), [PhysicsAsset = PhysicsAsset, bEnableCollision, &bResult](const FSelection& Lhs, const FSelection& Rhs)
 		{
-			if(PhysicsAsset->IsCollisionEnabled(SelectedBodies[i].Index, SelectedBodies[j].Index) != bEnableCollision)
+			if (PhysicsAsset->IsCollisionEnabled(Lhs.Index, Rhs.Index) != bEnableCollision)
 			{
-				return true;
+				bResult = true;
 			}
-		}
-	}
+		});
 
-	return false;
+	return bResult;
 }
 
 void FPhysicsAssetEditorSharedData::SetCollisionBetweenSelectedAndAll(bool bEnableCollision)
 {
-	if (bRunningSimulation || SelectedBodies.Num() == 0)
+	if (bRunningSimulation || UniqueSelectionReferencingBodies().IsEmpty())
 	{
 		return;
 	}
 
 	PhysicsAsset->Modify();
 
-	for(int32 i=0; i<SelectedBodies.Num(); ++i)
+	for(const FSelection& Selection : UniqueSelectionReferencingBodies())
 	{
 		for(int32 j = 0; j < PhysicsAsset->SkeletalBodySetups.Num(); ++j)
 		{
 			if(bEnableCollision)
 			{
-				PhysicsAsset->EnableCollision(SelectedBodies[i].Index, j);
+				PhysicsAsset->EnableCollision(Selection.Index, j);
 			}
 			else
 			{
-				PhysicsAsset->DisableCollision(SelectedBodies[i].Index, j);
+				PhysicsAsset->DisableCollision(Selection.Index, j);
 			}
 
 		}
 	}
 
 	UpdateNoCollisionBodies();
-
 	RefreshPhysicsAssetChange(PhysicsAsset);
-
 	BroadcastPreviewChanged();
 }
 
 bool FPhysicsAssetEditorSharedData::CanSetCollisionBetweenSelectedAndAll(bool bEnableCollision) const
 {
-	if (bRunningSimulation || SelectedBodies.Num() == 0)
-	{
-		return false;
-	}
-
-	for(int32 i=0; i<SelectedBodies.Num(); ++i)
+	if (!bRunningSimulation)
 	{
-		for(int32 j = 0; j < PhysicsAsset->SkeletalBodySetups.Num(); ++j)
+		for (const FSelection& SelectedBody : UniqueSelectionReferencingBodies())
 		{
-			if(PhysicsAsset->IsCollisionEnabled(SelectedBodies[i].Index, j) != bEnableCollision)
+			for (int32 j = 0; j < PhysicsAsset->SkeletalBodySetups.Num(); ++j)
 			{
-				return true;
+				if (PhysicsAsset->IsCollisionEnabled(SelectedBody.Index, j) != bEnableCollision)
+				{
+					return true;
+				}
 			}
 		}
 	}
@@ -1666,9 +1718,9 @@ void FPhysicsAssetEditorSharedData::SetPrimitiveCollision(ECollisionEnabled::Typ
 
 	PhysicsAsset->Modify();
 
-	for (FSelection SelectedBody : SelectedBodies)
+	for (const FSelection& SelectedBody : UniqueSelectionReferencingBodies())
 	{
-		PhysicsAsset->SetPrimitiveCollision(SelectedBody.Index, SelectedBody.PrimitiveType, SelectedBody.PrimitiveIndex, CollisionEnabled);
+		PhysicsAsset->SetPrimitiveCollision(SelectedBody.GetIndex(), SelectedBody.GetPrimitiveType(), SelectedBody.GetPrimitiveIndex(), CollisionEnabled);
 	}
 
 	BroadcastPreviewChanged();
@@ -1676,7 +1728,7 @@ void FPhysicsAssetEditorSharedData::SetPrimitiveCollision(ECollisionEnabled::Typ
 
 bool FPhysicsAssetEditorSharedData::CanSetPrimitiveCollision(ECollisionEnabled::Type CollisionEnabled) const
 {
-	if (bRunningSimulation || SelectedBodies.Num() == 0)
+	if (bRunningSimulation || UniqueSelectionReferencingBodies().IsEmpty())
 	{
 		return false;
 	}
@@ -1686,9 +1738,9 @@ bool FPhysicsAssetEditorSharedData::CanSetPrimitiveCollision(ECollisionEnabled::
 
 bool FPhysicsAssetEditorSharedData::GetIsPrimitiveCollisionEnabled(ECollisionEnabled::Type CollisionEnabled) const
 {
-	for (const FSelection& SelectedBody : SelectedBodies)
+	for (const FSelection& Selection : SelectedPrimitives())
 	{
-		if (PhysicsAsset->GetPrimitiveCollision(SelectedBody.Index, SelectedBody.PrimitiveType, SelectedBody.PrimitiveIndex) == CollisionEnabled)
+		if (PhysicsAsset->GetPrimitiveCollision(Selection.GetIndex(), Selection.GetPrimitiveType(), Selection.GetPrimitiveIndex()) == CollisionEnabled)
 		{
 			return true;
 		}
@@ -1699,9 +1751,9 @@ bool FPhysicsAssetEditorSharedData::GetIsPrimitiveCollisionEnabled(ECollisionEna
 
 void FPhysicsAssetEditorSharedData::SetPrimitiveContributeToMass(bool bContributeToMass)
 {
-	for (const FSelection& SelectedBody : SelectedBodies)
+	for (const FSelection& Selection : SelectedPrimitives())
 	{
-		PhysicsAsset->SetPrimitiveContributeToMass(SelectedBody.Index, SelectedBody.PrimitiveType, SelectedBody.PrimitiveIndex, bContributeToMass);
+		PhysicsAsset->SetPrimitiveContributeToMass(Selection.Index, Selection.GetPrimitiveType(), Selection.PrimitiveIndex, bContributeToMass);
 	}
 }
 
@@ -1712,9 +1764,9 @@ bool FPhysicsAssetEditorSharedData::CanSetPrimitiveContributeToMass() const
 
 bool FPhysicsAssetEditorSharedData::GetPrimitiveContributeToMass() const
 {
-	for (const FSelection& SelectedBody : SelectedBodies)
+	for (const FSelection& Selection : SelectedPrimitives())
 	{
-		if (PhysicsAsset->GetPrimitiveContributeToMass(SelectedBody.Index, SelectedBody.PrimitiveType, SelectedBody.PrimitiveIndex))
+		if (PhysicsAsset->GetPrimitiveContributeToMass(Selection.Index, Selection.GetPrimitiveType(), Selection.PrimitiveIndex))
 		{
 			return true;
 		}
@@ -1887,15 +1939,15 @@ void FPhysicsAssetEditorSharedData::CopySelectedBodiesAndConstraintsToClipboard(
 			OutNumCopiedBodies = 0;
 
 			// Export each of the selected nodes
-			for (const FSelection& SelectedBody : SelectedBodies)
+			for (const FSelection& Selection : UniqueSelectionReferencingBodies())
 			{
 				// selected bodies contain the primitives, so abody can be stored multiple time for each of its primitive
 				// we need to make sure we process it only once
-				if (!ExportedBodyIndices.Contains(SelectedBody.Index))
+				if (!ExportedBodyIndices.Contains(Selection.Index))
 				{
-					ExportedBodyIndices.Add(SelectedBody.Index);
+					ExportedBodyIndices.Add(Selection.Index);
 
-					if (USkeletalBodySetup* BodySetup = PhysicsAsset->SkeletalBodySetups[SelectedBody.Index])
+					if (USkeletalBodySetup* BodySetup = PhysicsAsset->SkeletalBodySetups[Selection.Index])
 					{						
 						UExporter::ExportToOutputDevice(&Context, BodySetup, NULL, Archive, TEXT("copy"), 0, PPF_ExportsNotFullyQualified | PPF_Copy | PPF_Delimited, false);
 						++OutNumCopiedBodies;
@@ -1910,7 +1962,7 @@ void FPhysicsAssetEditorSharedData::CopySelectedBodiesAndConstraintsToClipboard(
 			TSet<int32> ExportedConstraintIndices;
 
 			// Export each of the selected nodes
-			for (const FSelection& SelectedConstraint : SelectedConstraints)
+			for (const FSelection& SelectedConstraint : SelectedConstraints())
 			{
 				// selected bodies contain the primitives, so a body can be stored multiple time for each of its primitive
 				// we need to make sure we process it only once
@@ -2104,7 +2156,7 @@ void FPhysicsAssetEditorSharedData::PasteBodiesAndConstraintsFromClipboard(int32
 					}
 
 					// Enable collisions between all pasted bodies.
-					ForEachUniquePair(PastedBodyIndices, [this](const int32 BodyIndexA, const int32 BodyIndexB) 
+					ForEachUniquePair(PastedBodyIndices, [this](const int32 BodyIndexA, const int32 BodyIndexB)
 						{ 
 							PhysicsAsset->EnableCollision(BodyIndexA, BodyIndexB); 
 						});
@@ -2148,13 +2200,13 @@ void FPhysicsAssetEditorSharedData::CopySelectedShapesToClipboard(int32& OutNumC
 		NewBodySetup->AddToRoot();
 		{
 			TSet<int32> SelectedBodyIndices;
-			for (const FSelection& SelectedBody : SelectedBodies)
+			for (const FSelection& Selection : SelectedPrimitives())
 			{
-				if (const USkeletalBodySetup* OldBodySetup = PhysicsAsset->SkeletalBodySetups[SelectedBody.Index])
+				if (const USkeletalBodySetup* OldBodySetup = PhysicsAsset->SkeletalBodySetups[Selection.Index])
 				{
-					if (NewBodySetup->AddCollisionElemFrom(OldBodySetup->AggGeom, SelectedBody.PrimitiveType, SelectedBody.PrimitiveIndex))
+					if (NewBodySetup->AddCollisionElemFrom(OldBodySetup->AggGeom, Selection.GetPrimitiveType(), Selection.GetPrimitiveIndex()))
 					{
-						SelectedBodyIndices.Add(SelectedBody.Index);
+						SelectedBodyIndices.Add(Selection.Index);
 						++OutNumCopiedShapes;
 					}
 				}
@@ -2203,17 +2255,17 @@ void FPhysicsAssetEditorSharedData::PasteShapesFromClipboard(int32& OutNumPasted
 				Factory.ProcessBuffer(TempPackage, RF_Transactional, TextToImport);
 
 				// Paste copied shapes into each of the selected bodies
-				if (Factory.NewBodySetups.Num() > 0 && SelectedBodies.Num() > 0)
+				if (Factory.NewBodySetups.Num() > 0 && !UniqueSelectionReferencingBodies().IsEmpty())
 				{
 					const FScopedTransaction Transaction(NSLOCTEXT("PhysicsAssetEditor", "PasteShapesFromClipboard", "Paste Shapes From Clipboard"));
 
 					// We have to track which bodies we've pasted into, because they might appear multiple times
-					// (for separate primitive shapes) in the SelectedBodies list.
+					// (for separate primitive shapes) in the SelectedObjects->SelectedBodies() list.
 					TSet<int32> PastedBodyIndices;
 					for (const UBodySetup* NewBodySetup : Factory.NewBodySetups)
 					{
 						OutNumPastedShapes += NewBodySetup->AggGeom.GetElementCount();
-						for (const FSelection& SelectedBody : SelectedBodies)
+						for (const FSelection& SelectedBody : UniqueSelectionReferencingBodies())
 						{
 							if (!PastedBodyIndices.Contains(SelectedBody.Index))
 							{
@@ -2241,8 +2293,8 @@ void FPhysicsAssetEditorSharedData::PasteShapesFromClipboard(int32& OutNumPasted
 
 void FPhysicsAssetEditorSharedData::CopyBodyProperties()
 {
-	check(SelectedBodies.Num() == 1);
-	CopyToClipboard(SharedDataConstants::BodyType, PhysicsAsset->SkeletalBodySetups[GetSelectedBody()->Index]);
+	check(UniqueSelectionReferencingBodies().Num() == 1);
+	CopyToClipboard(SharedDataConstants::BodyType, PhysicsAsset->SkeletalBodySetups[GetSelectedBodyOrPrimitive()->Index]);
 }
 
 void FPhysicsAssetEditorSharedData::PasteBodyProperties()
@@ -2270,15 +2322,15 @@ void FPhysicsAssetEditorSharedData::PasteBodyProperties()
 		return;
 	}
 
-	if(SelectedBodies.Num() > 0)
+	if(!UniqueSelectionReferencingBodies().IsEmpty())
 	{
 		const FScopedTransaction Transaction( NSLOCTEXT("PhysicsAssetEditor", "PasteBodyProperties", "Paste Body Properties") );
 
 		PhysicsAsset->Modify();
 
-		for(int32 i=0; i<SelectedBodies.Num(); ++i)
+		for (const FSelection& Selection : UniqueSelectionReferencingBodies())
 		{
-			UBodySetup* ToBodySetup = PhysicsAsset->SkeletalBodySetups[SelectedBodies[i].Index];
+			UBodySetup* const ToBodySetup = PhysicsAsset->SkeletalBodySetups[Selection.Index];
 			ToBodySetup->Modify();
 			ToBodySetup->CopyBodyPropertiesFrom(CopiedBodySetup);
 		}
@@ -2290,8 +2342,8 @@ void FPhysicsAssetEditorSharedData::PasteBodyProperties()
 
 void FPhysicsAssetEditorSharedData::CopyBodyName()
 {
-	check(SelectedBodies.Num() == 1);
-	FPlatformApplicationMisc::ClipboardCopy(*PhysicsAsset->SkeletalBodySetups[GetSelectedBody()->Index]->BoneName.ToString());
+	check(UniqueSelectionReferencingBodies().Num() == 1);
+	FPlatformApplicationMisc::ClipboardCopy(*PhysicsAsset->SkeletalBodySetups[GetSelectedBodyOrPrimitive()->Index]->BoneName.ToString());
 }
 
 bool FPhysicsAssetEditorSharedData::WeldSelectedBodies(bool bWeld /* = true */)
@@ -2302,7 +2354,7 @@ bool FPhysicsAssetEditorSharedData::WeldSelectedBodies(bool bWeld /* = true */)
 		return false;
 	}
 
-	if(SelectedBodies.Num() <= 1)
+	if(UniqueSelectionReferencingBodies().Num() <= 1)
 	{
 		return false;
 	}
@@ -2317,21 +2369,27 @@ bool FPhysicsAssetEditorSharedData::WeldSelectedBodies(bool bWeld /* = true */)
 	int BodyIndex0 = 0;
 	int BodyIndex1 = INDEX_NONE;
 
-	for(int32 i=1; i<SelectedBodies.Num(); ++i)
+	for(int SelectedIndex = 0, SelectedCount = SelectedObjects->Num(); SelectedIndex < SelectedCount; ++SelectedIndex)
 	{
-		if(SelectedBodies[BodyIndex0].Index == SelectedBodies[i].Index)
-		{
-			continue;
-		}
+		const FSelection& SelectedElement = SelectedObjects->GetSelectedAt(SelectedIndex);
 
-		if(BodyIndex1== INDEX_NONE)
-		{
-			BodyIndex1 = i;
-		}else
+		if (SelectedElement.HasType(FPhysicsAssetEditorSelectedElement::Body))
 		{
-			if(SelectedBodies[BodyIndex1].Index != SelectedBodies[i].Index)
+			if (SelectedObjects->GetSelectedAt(BodyIndex0).Index == SelectedElement.Index)
+			{
+				continue;
+			}
+
+			if (BodyIndex1 == INDEX_NONE)
+			{
+				BodyIndex1 = SelectedIndex;
+			}
+			else
 			{
-				return false;
+				if (SelectedObjects->GetSelectedAt(BodyIndex1).Index != SelectedElement.Index)
+				{
+					return false;
+				}
 			}
 		}
 	}
@@ -2342,8 +2400,11 @@ bool FPhysicsAssetEditorSharedData::WeldSelectedBodies(bool bWeld /* = true */)
 		return false;
 	}
 
-	const FSelection& Body0 = SelectedBodies[BodyIndex0];
-	const FSelection& Body1 = SelectedBodies[BodyIndex1];
+	check(SelectedObjects->IsValidIndex(BodyIndex0));
+	check(SelectedObjects->IsValidIndex(BodyIndex1));
+
+	const FSelection& Body0 = SelectedObjects->GetSelectedAt(BodyIndex0);
+	const FSelection& Body1 = SelectedObjects->GetSelectedAt(BodyIndex1);
 
 	FName Bone0Name = PhysicsAsset->SkeletalBodySetups[Body0.Index]->BoneName;
 	int32 Bone0Index = EditorSkelMesh->GetRefSkeleton().FindBoneIndex(Bone0Name);
@@ -2364,28 +2425,28 @@ bool FPhysicsAssetEditorSharedData::WeldSelectedBodies(bool bWeld /* = true */)
 	int32 ParentPrimitiveIndex = INDEX_NONE;
 	int32 ChildPrimitiveIndex = INDEX_NONE;
 
-	if (PhysicsAsset->FindControllingBodyIndex(EditorSkelMesh, Bone1ParentIndex) == Body0.Index)
+	if (PhysicsAsset->FindControllingBodyIndex(EditorSkelMesh, Bone1ParentIndex) == Body0.GetIndex())
 	{
-		ParentBodyIndex = Body0.Index;
+		ParentBodyIndex = Body0.GetIndex();
 		ParentBoneName = Bone0Name;
-		ChildBodyIndex = Body1.Index;
-		ParentPrimitiveType = Body0.PrimitiveType;
-		ChildPrimitiveType = Body1.PrimitiveType;
-		ParentPrimitiveIndex = Body0.PrimitiveIndex;
+		ChildBodyIndex = Body1.GetIndex();
+		ParentPrimitiveType = Body0.GetPrimitiveType();
+		ChildPrimitiveType = Body1.GetPrimitiveType();
+		ParentPrimitiveIndex = Body0.GetPrimitiveIndex();
 		//Child geoms get appended so just add it. This is kind of a hack but this whole indexing scheme needs to be rewritten anyway
-		ChildPrimitiveIndex = Body1.PrimitiveIndex + PhysicsAsset->SkeletalBodySetups[Body0.Index]->AggGeom.GetElementCount(ChildPrimitiveType);
+		ChildPrimitiveIndex = Body1.GetPrimitiveIndex() + PhysicsAsset->SkeletalBodySetups[Body0.Index]->AggGeom.GetElementCount(ChildPrimitiveType);
 
 		bCanWeld = true;
-	}else if(PhysicsAsset->FindControllingBodyIndex(EditorSkelMesh, Bone0ParentIndex) == Body1.Index)
+	}else if(PhysicsAsset->FindControllingBodyIndex(EditorSkelMesh, Bone0ParentIndex) == Body1.GetIndex())
 	{
-		ParentBodyIndex = Body1.Index;
+		ParentBodyIndex = Body1.GetIndex();
 		ParentBoneName = Bone1Name;
-		ChildBodyIndex = Body0.Index;
-		ParentPrimitiveType = Body1.PrimitiveType;
-		ChildPrimitiveType = Body0.PrimitiveType;
-		ParentPrimitiveIndex = Body1.PrimitiveIndex;
+		ChildBodyIndex = Body0.GetIndex();
+		ParentPrimitiveType = Body1.GetPrimitiveType();
+		ChildPrimitiveType = Body0.GetPrimitiveType();
+		ParentPrimitiveIndex = Body1.GetPrimitiveIndex();
 		//Child geoms get appended so just add it. This is kind of a hack but this whole indexing scheme needs to be rewritten anyway
-		ChildPrimitiveIndex = Body0.PrimitiveIndex + PhysicsAsset->SkeletalBodySetups[Body1.Index]->AggGeom.GetElementCount(ChildPrimitiveType);
+		ChildPrimitiveIndex = Body0.GetPrimitiveIndex() + PhysicsAsset->SkeletalBodySetups[Body1.GetIndex()]->AggGeom.GetElementCount(ChildPrimitiveType);
 
 		bCanWeld = true;
 	}
@@ -2426,21 +2487,30 @@ bool FPhysicsAssetEditorSharedData::WeldSelectedBodies(bool bWeld /* = true */)
 	// update the tree
 	BroadcastHierarchyChanged();
 
-	// Just to be safe - deselect any selected constraints
-	ClearSelectedConstraints();
-	ClearSelectedBody(); // Previous selection is invalid because child no longer has same index.
+	const int32 BodyIndex = PhysicsAsset->FindBodyIndex(ParentBoneName);
 
-	int32 BodyIndex = PhysicsAsset->FindBodyIndex(ParentBoneName);
-	FSelection SelectionParent(BodyIndex, ParentPrimitiveType, ParentPrimitiveIndex);
-	SetSelectedBody(SelectionParent, true); // This redraws the viewport as well...
-
-	FSelection SelectionChild(BodyIndex, ChildPrimitiveType, ChildPrimitiveIndex);
-	SetSelectedBody(SelectionChild, true); // This redraws the viewport as well...
+	// Previous selection is invalid because child no longer has same index. Just to be safe - deselect any selected bodies or constraints
+	SetSelectedPrimitives({ MakePrimitiveSelection(BodyIndex, ParentPrimitiveType, ParentPrimitiveIndex), MakePrimitiveSelection(BodyIndex, ChildPrimitiveType, ChildPrimitiveIndex) }); // This redraws the viewport as well...	
 
 	RefreshPhysicsAssetChange(PhysicsAsset);
 	return true;
 }
 
+bool FPhysicsAssetEditorSharedData::ModifySelectionInternal(TFunctionRef<bool(void)> SelectionOperation)
+{
+	if (!InsideSelChange && SelectionOperation())
+	{
+		BroadcastSelectionChanged();
+
+		++InsideSelChange;
+		BroadcastPreviewChanged();
+		--InsideSelChange;
+
+		return true;
+	}
+
+	return false;
+}
 
 void FPhysicsAssetEditorSharedData::InitConstraintSetup(UPhysicsConstraintTemplate* ConstraintSetup, int32 ChildBodyIndex, int32 ParentBodyIndex)
 {
@@ -2490,11 +2560,13 @@ void FPhysicsAssetEditorSharedData::RecreateBody(const FPhysAssetCreateParams& B
 
 		if (bAutoSelect)
 		{
-			SetSelectedBodyAnyPrimitive(BodyIndex, true);
+			ModifySelectedBodies(BodyIndex, true);
 		}
 
 		RefreshPhysicsAssetChange(PhysicsAsset);
 	}
+
+	RefreshPhysicsAssetChange(PhysicsAsset);
 }
 
 void FPhysicsAssetEditorSharedData::MakeNewBody(int32 NewBoneIndex, bool bAutoSelect)
@@ -2639,7 +2711,7 @@ void FPhysicsAssetEditorSharedData::MakeNewBody(const FPhysAssetCreateParams& Ne
 
 	if (bAutoSelect)
 	{
-		SetSelectedBodyAnyPrimitive(NewBodyIndex, true);
+		ModifySelectedBodies(NewBodyIndex, true);
 	}
 	
 
@@ -2702,8 +2774,7 @@ void FPhysicsAssetEditorSharedData::MakeNewConstraints(int32 ParentBodyIndex, co
 		}
 	}
 
-	ClearSelectedConstraints();
-	SetSelectedConstraints(NewlyCreatedConstraints, true);
+	SetSelectedConstraints(NewlyCreatedConstraints);
 
 	// update the tree
 	BroadcastHierarchyChanged();
@@ -2754,7 +2825,7 @@ void FPhysicsAssetEditorSharedData::SnapConstraintToBone(FConstraintInstance& Co
 
 void FPhysicsAssetEditorSharedData::CopyConstraintProperties()
 {
-	check(SelectedConstraints.Num() == 1);
+	check(SelectedConstraints().Num() == 1);
 	CopyToClipboard(SharedDataConstants::ConstraintType, PhysicsAsset->ConstraintSetup[GetSelectedConstraint()->Index]);
 }
 
@@ -2768,15 +2839,17 @@ void FPhysicsAssetEditorSharedData::PasteConstraintProperties()
 		return;
 	}
 
-	const UPhysicsConstraintTemplate* FromConstraintSetup = Cast<UPhysicsConstraintTemplate>(SourceConstraint);
+	const UPhysicsConstraintTemplate* const FromConstraintSetup = Cast<UPhysicsConstraintTemplate>(SourceConstraint);
+
+	FPhysicsAssetEditorSharedData::SelectionFilterRange SelectedConstraintRange = SelectedConstraints();
 
-	if(FromConstraintSetup && SelectedConstraints.Num() > 0)
+	if(FromConstraintSetup && !SelectedConstraintRange.IsEmpty())
 	{
 		const FScopedTransaction Transaction(NSLOCTEXT("PhysicsAssetEditor", "PasteConstraintProperties", "Paste Constraint Properties"));
 
-		for(int32 i=0; i<SelectedConstraints.Num(); ++i)
+		for (const FSelection& SelectedConstraint : SelectedConstraintRange)
 		{
-			UPhysicsConstraintTemplate* ToConstraintSetup = PhysicsAsset->ConstraintSetup[SelectedConstraints[i].Index];
+			UPhysicsConstraintTemplate* const ToConstraintSetup = PhysicsAsset->ConstraintSetup[SelectedConstraint.Index];
 			CopyConstraintProperties(FromConstraintSetup, ToConstraintSetup, /*bKeepOriginalRotation=*/true);
 		}
 	}
@@ -2796,13 +2869,13 @@ void FPhysicsAssetEditorSharedData::CycleCurrentConstraintOrientation()
 {
 	const FScopedTransaction Transaction( LOCTEXT("CycleCurrentConstraintOrientation", "Cycle Current Constraint Orientation") );
 
-	for(int32 i=0; i<SelectedConstraints.Num(); ++i)
+	for (const FSelection& SelectedConstraint : SelectedConstraints())
 	{
-		UPhysicsConstraintTemplate* ConstraintTemplate = PhysicsAsset->ConstraintSetup[SelectedConstraints[i].Index];
+		UPhysicsConstraintTemplate* ConstraintTemplate = PhysicsAsset->ConstraintSetup[SelectedConstraint.Index];
 		ConstraintTemplate->Modify();
 		FMatrix ConstraintTransform = ConstraintTemplate->DefaultInstance.GetRefFrame(EConstraintFrame::Frame2).ToMatrixWithScale();
-		FTransform WParentFrame = GetConstraintWorldTM(&SelectedConstraints[i], EConstraintFrame::Frame2);
-		FTransform WChildFrame = GetConstraintWorldTM(&SelectedConstraints[i], EConstraintFrame::Frame1);
+		FTransform WParentFrame = GetConstraintWorldTM(&SelectedConstraint, EConstraintFrame::Frame2);
+		FTransform WChildFrame = GetConstraintWorldTM(&SelectedConstraint, EConstraintFrame::Frame1);
 		FTransform RelativeTransform = WChildFrame * WParentFrame.Inverse();
 
 		CycleMatrixRows(&ConstraintTransform);
@@ -2816,11 +2889,11 @@ void FPhysicsAssetEditorSharedData::CycleCurrentConstraintActive()
 {
 	const FScopedTransaction Transaction( LOCTEXT("CycleCurrentConstraintActive", "Cycle Current Constraint Active") );
 
-	for(int32 i=0; i<SelectedConstraints.Num(); ++i)
+	for (const FSelection& SelectedConstraint : SelectedConstraints())
 	{
-		UPhysicsConstraintTemplate* ConstraintTemplate = PhysicsAsset->ConstraintSetup[SelectedConstraints[i].Index];
+		UPhysicsConstraintTemplate* const ConstraintTemplate = PhysicsAsset->ConstraintSetup[SelectedConstraint.Index];
 		ConstraintTemplate->Modify();
-		FConstraintInstance & DefaultInstance = ConstraintTemplate->DefaultInstance;
+		FConstraintInstance& DefaultInstance = ConstraintTemplate->DefaultInstance;
 
 		if(DefaultInstance.GetAngularSwing1Motion() != ACM_Limited && DefaultInstance.GetAngularSwing2Motion() != ACM_Limited)
 		{
@@ -2847,9 +2920,9 @@ void FPhysicsAssetEditorSharedData::ToggleConstraint(EPhysicsAssetEditorConstrai
 {
 	const FScopedTransaction Transaction( LOCTEXT("ToggleConstraintTypeLock", "Toggle Constraint Type Lock") );
 
-	for(int32 i=0; i<SelectedConstraints.Num(); ++i)
+	for (const FSelection& SelectedConstraint : SelectedConstraints())
 	{
-		UPhysicsConstraintTemplate* ConstraintTemplate = PhysicsAsset->ConstraintSetup[GetSelectedConstraint()->Index];
+		UPhysicsConstraintTemplate* const ConstraintTemplate = PhysicsAsset->ConstraintSetup[SelectedConstraint.Index];
 		ConstraintTemplate->Modify();
 		FConstraintInstance & DefaultInstance = ConstraintTemplate->DefaultInstance;
 
@@ -2873,10 +2946,10 @@ bool FPhysicsAssetEditorSharedData::IsAngularConstraintLocked(EPhysicsAssetEdito
 	bool bLocked = false;
 	bool bSame = false;
 
-	for(int32 i = 0; i < SelectedConstraints.Num(); ++i)
+	for (const FSelection& SelectedConstraint : SelectedConstraints())
 	{
-		UPhysicsConstraintTemplate* ConstraintTemplate = PhysicsAsset->ConstraintSetup[GetSelectedConstraint()->Index];
-		FConstraintInstance & DefaultInstance = ConstraintTemplate->DefaultInstance;
+		const UPhysicsConstraintTemplate* const ConstraintTemplate = PhysicsAsset->ConstraintSetup[SelectedConstraint.Index];
+		const FConstraintInstance& DefaultInstance = ConstraintTemplate->DefaultInstance;
 
 		if(Constraint == PCT_Swing1)
 		{
@@ -2981,21 +3054,12 @@ void FPhysicsAssetEditorSharedData::DeleteCurrentSelection()
 void FPhysicsAssetEditorSharedData::DeleteCurrentBody()
 {
 	// Delete any directly selected bodies and all their primitives.
-
-	TArray<FSelection> DirectSelectedBodies;
-
-	for (const FSelection& Selection : SelectedBodies)
-	{
-		if (Selection.bIsDirectBodySelection)
-		{
-			DirectSelectedBodies.Add(Selection);
-		}
-	}
+	TArray<FSelection> DirectSelectedBodies = SelectedBodies().ToArray();
 
 	if (!DirectSelectedBodies.IsEmpty())
 	{
 		// Remove target body indexes from the selection.
-		SetSelectedBodies(DirectSelectedBodies, false);
+		ModifySelectedBodies(DirectSelectedBodies, false);
 
 		// Sort by body index - highest first - as body indexes greater than the deleted index in the physics asset will be modified by each deletion.
 		Algo::SortBy(DirectSelectedBodies, &FSelection::Index, TGreater<>());
@@ -3018,7 +3082,7 @@ void FPhysicsAssetEditorSharedData::DeleteCurrentPrim()
 		return;
 	}
 
-	if (!GetSelectedBody())
+	if (!GetSelectedBodyOrPrimitive())
 	{
 		return;
 	}
@@ -3030,10 +3094,10 @@ void FPhysicsAssetEditorSharedData::DeleteCurrentPrim()
 	//We need to do this first because deleting will modify our selection
 	TMap<UBodySetup *, TArray<FSelection>> BodySelectionMap;
 	TArray<UBodySetup*> BodySetups;
-	for(int32 i=0; i<SelectedBodies.Num(); ++i)
+	for (const FSelection& SelectedPrimitive : SelectedPrimitives())
 	{
-		UBodySetup* BodySetup = PhysicsAsset->SkeletalBodySetups[SelectedBodies[i].Index];
-		BodySelectionMap.FindOrAdd(BodySetup).Add(SelectedBodies[i]);
+		UBodySetup* const BodySetup = PhysicsAsset->SkeletalBodySetups[SelectedPrimitive.Index];
+		BodySelectionMap.FindOrAdd(BodySetup).Add(SelectedPrimitive);
 	}
 
 	const FScopedTransaction Transaction( NSLOCTEXT("UnrealEd", "DeletePrimitive", "Delete Primitive") );
@@ -3056,25 +3120,25 @@ void FPhysicsAssetEditorSharedData::DeleteCurrentPrim()
 
 			BodySetup->Modify();
 
-			if (SelectedBody.PrimitiveType == EAggCollisionShape::Sphere)
+			if (SelectedBody.GetPrimitiveType() == EAggCollisionShape::Sphere)
 			{
 				BodySetup->AggGeom.SphereElems.RemoveAt(SelectedBody.PrimitiveIndex - (SphereDeletedCount++));
 			}
-			else if (SelectedBody.PrimitiveType == EAggCollisionShape::Box)
+			else if (SelectedBody.GetPrimitiveType() == EAggCollisionShape::Box)
 			{
 				BodySetup->AggGeom.BoxElems.RemoveAt(SelectedBody.PrimitiveIndex - (BoxDeletedCount++));
 			}
-			else if (SelectedBody.PrimitiveType == EAggCollisionShape::Sphyl)
+			else if (SelectedBody.GetPrimitiveType() == EAggCollisionShape::Sphyl)
 			{
 				BodySetup->AggGeom.SphylElems.RemoveAt(SelectedBody.PrimitiveIndex - (SphylDeletedCount++));
 			}
-			else if (SelectedBody.PrimitiveType == EAggCollisionShape::Convex)
+			else if (SelectedBody.GetPrimitiveType() == EAggCollisionShape::Convex)
 			{
 				BodySetup->AggGeom.ConvexElems.RemoveAt(SelectedBody.PrimitiveIndex - (ConvexDeletedCount++));
 				// Need to invalidate GUID in this case as cooked data must be updated
 				BodySetup->InvalidatePhysicsData();
 			}
-			else if (SelectedBody.PrimitiveType == EAggCollisionShape::TaperedCapsule)
+			else if (SelectedBody.GetPrimitiveType() == EAggCollisionShape::TaperedCapsule)
 			{
 				BodySetup->AggGeom.TaperedCapsuleElems.RemoveAt(SelectedBody.PrimitiveIndex - (TaperedCapsuleDeletedCount++));
 			}
@@ -3165,10 +3229,11 @@ void FPhysicsAssetEditorSharedData::DeleteCurrentConstraint()
 
 	//Save indices before delete because delete modifies our Selected array
 	TArray<int32> Indices;
-	for(int32 i=0; i<SelectedConstraints.Num(); ++i)
+
+	for (const FSelection& SelectedConstraint : SelectedConstraints())
 	{
-		ConditionalClearClipboard(SharedDataConstants::ConstraintType, PhysicsAsset->ConstraintSetup[SelectedConstraints[i].Index]);
-		Indices.Add(SelectedConstraints[i].Index);
+		ConditionalClearClipboard(SharedDataConstants::ConstraintType, PhysicsAsset->ConstraintSetup[SelectedConstraint.Index]);
+		Indices.Add(SelectedConstraint.Index);
 	}
 
 	Indices.Sort();
@@ -3189,7 +3254,7 @@ void FPhysicsAssetEditorSharedData::DeleteCurrentConstraint()
 void FPhysicsAssetEditorSharedData::ToggleSimulation()
 {
 	// don't start simulation if there are no bodies or if we are manipulating a body
-	if (PhysicsAsset->SkeletalBodySetups.Num() == 0 || bManipulating)
+	if (PhysicsAsset->SkeletalBodySetups.Num() == 0 || IsManipulating())
 	{
 		return;  
 	}
@@ -3398,27 +3463,29 @@ TSharedRef<SWidget> FPhysicsAssetEditorSharedData::CreateGenerateBodiesWidget(co
 
 void FPhysicsAssetEditorSharedData::PostUndo()
 {
+	// The selection can become invalid if the creation of an object that is selected is undone etc - try to detect that here and clear selection if it is the case.
+
 	bool bInvalidSelection = false;
 
-	for (int32 BodyIndex = 0; BodyIndex < SelectedBodies.Num() && bInvalidSelection == false; ++BodyIndex)
+	for(UPhysicsAssetEditorSelection::Iterator Itr = SelectedBodiesAndPrimitives().CreateConstIterator(); Itr && !bInvalidSelection; ++Itr)
 	{
-		const FSelection& Selection = SelectedBodies[BodyIndex];
-		if (PhysicsAsset->SkeletalBodySetups.Num() <= Selection.Index)
+		const FSelection& Selection = *Itr;
+
+		if (PhysicsAsset->SkeletalBodySetups.Num() <= Selection.GetIndex())
 		{
 			bInvalidSelection = true;
 		}
 		else
-		{
-			
-			if (UBodySetup * BodySetup = PhysicsAsset->SkeletalBodySetups[Selection.Index])
+		{		
+			if (UBodySetup * BodySetup = PhysicsAsset->SkeletalBodySetups[Selection.GetIndex()])
 			{
-				switch (Selection.PrimitiveType)
+				switch (Selection.GetPrimitiveType())
 				{
-				case EAggCollisionShape::Box: bInvalidSelection = BodySetup->AggGeom.BoxElems.Num() <= Selection.PrimitiveIndex ? true : bInvalidSelection; break;
-				case EAggCollisionShape::Convex: bInvalidSelection = BodySetup->AggGeom.ConvexElems.Num() <= Selection.PrimitiveIndex ? true : bInvalidSelection; break;
-				case EAggCollisionShape::Sphere: bInvalidSelection = BodySetup->AggGeom.SphereElems.Num() <= Selection.PrimitiveIndex ? true : bInvalidSelection; break;
-				case EAggCollisionShape::Sphyl: bInvalidSelection = BodySetup->AggGeom.SphylElems.Num() <= Selection.PrimitiveIndex ? true : bInvalidSelection; break;
-				case EAggCollisionShape::TaperedCapsule: bInvalidSelection = BodySetup->AggGeom.TaperedCapsuleElems.Num() <= Selection.PrimitiveIndex ? true : bInvalidSelection; break;
+				case EAggCollisionShape::Box: bInvalidSelection = BodySetup->AggGeom.BoxElems.Num() <= Selection.GetPrimitiveIndex() ? true : bInvalidSelection; break;
+				case EAggCollisionShape::Convex: bInvalidSelection = BodySetup->AggGeom.ConvexElems.Num() <= Selection.GetPrimitiveIndex() ? true : bInvalidSelection; break;
+				case EAggCollisionShape::Sphere: bInvalidSelection = BodySetup->AggGeom.SphereElems.Num() <= Selection.GetPrimitiveIndex() ? true : bInvalidSelection; break;
+				case EAggCollisionShape::Sphyl: bInvalidSelection = BodySetup->AggGeom.SphylElems.Num() <= Selection.GetPrimitiveIndex() ? true : bInvalidSelection; break;
+				case EAggCollisionShape::TaperedCapsule: bInvalidSelection = BodySetup->AggGeom.TaperedCapsuleElems.Num() <= Selection.GetPrimitiveIndex() ? true : bInvalidSelection; break;
 				default: bInvalidSelection = true;
 				}
 			}
@@ -3429,9 +3496,9 @@ void FPhysicsAssetEditorSharedData::PostUndo()
 		}
 	}
 
-	for (int32 ConstraintIndex = 0; ConstraintIndex < SelectedConstraints.Num() && bInvalidSelection == false; ++ConstraintIndex)
+	for(UPhysicsAssetEditorSelection::Iterator Itr = SelectedConstraints().CreateConstIterator(); Itr && (bInvalidSelection == false); ++Itr)
 	{
-		const FSelection& Selection = SelectedConstraints[ConstraintIndex];
+		const FSelection& Selection = *Itr;
 		if (PhysicsAsset->ConstraintSetup.Num() <= Selection.Index)
 		{
 			bInvalidSelection = true;
@@ -3441,12 +3508,12 @@ void FPhysicsAssetEditorSharedData::PostUndo()
 	if (bInvalidSelection)
 	{
 		// Clear selection before we undo. We don't transact the editor itself - don't want to have something selected that is then removed.
-		ClearSelectedBody();
-		ClearSelectedConstraints();
+		SelectedObjects->ClearSelectionWithoutTransaction(FSelection::Body | FSelection::Constraint);
 	}
 
 	BroadcastPreviewChanged();
 	BroadcastHierarchyChanged();
+	BroadcastSelectionChanged();
 }
 
 void FPhysicsAssetEditorSharedData::Redo()
@@ -3464,6 +3531,7 @@ void FPhysicsAssetEditorSharedData::Redo()
 
 	BroadcastPreviewChanged();
 	BroadcastHierarchyChanged();
+	BroadcastSelectionChanged();
 }
 
 void FPhysicsAssetEditorSharedData::AddReferencedObjects(FReferenceCollector& Collector)
@@ -3503,22 +3571,22 @@ void FPhysicsAssetEditorSharedData::UpdateClothPhysics()
 	}
 }
 
-TArray<FPhysicsAssetEditorSharedData::FSelection> FPhysicsAssetEditorSharedData::GetUniqueSelectedBodies() const
+FVector FPhysicsAssetEditorSharedData::GetSelectedCoMPosition()
 {
-	TArray<FSelection> UniqueSelectedBodies;
-
-	TArray<int32> EncounteredBodyIndexs;
-
-	for (const FSelection& Selection : SelectedBodies)
+	if (const FSelection* const SelectedCoM = GetSelectedCoM())
 	{
-		if (!EncounteredBodyIndexs.Contains(Selection.Index))
+		if (const FVector* const ManipulatedBodyCoMPosition = FindManipulatedBodyCoMPosition(SelectedCoM->Index))
+		{
+			// return the CoM position from the FSelection object because the physics body's CoM position will only be updated at the end of manipulation.
+			return *ManipulatedBodyCoMPosition;
+		}
+		else
 		{
-			UniqueSelectedBodies.Add(Selection);
-			EncounteredBodyIndexs.Add(Selection.Index);
+			return EditorSkelComp->Bodies[SelectedCoM->Index]->GetCOMPosition();
 		}
 	}
 
-	return UniqueSelectedBodies;
+	return FVector::ZeroVector;
 }
 
 
@@ -3527,5 +3595,19 @@ FPhysicsAssetRenderSettings* FPhysicsAssetEditorSharedData::GetRenderSettings()
 	return UPhysicsAssetRenderUtilities::GetSettings(PhysicsAsset);
 }
 
+void FPhysicsAssetEditorSharedData::BeginManipulation()
+{
+	RecordSelectedCoM();
+
+	bManipulating = true;
+}
+
+void FPhysicsAssetEditorSharedData::EndManipulation()
+{
+	bManipulating = false;
+	bShouldUpdatedSelectedCoMs = true;
+
+	RefreshPhysicsAssetChange(PhysicsAsset, false);
+}
 
 #undef LOCTEXT_NAMESPACE
diff --git a/Engine/Source/Editor/PhysicsAssetEditor/Private/PhysicsAssetEditorSharedData.h b/Engine/Source/Editor/PhysicsAssetEditor/Private/PhysicsAssetEditorSharedData.h
--- a/Engine/Source/Editor/PhysicsAssetEditor/Private/PhysicsAssetEditorSharedData.h
+++ b/Engine/Source/Editor/PhysicsAssetEditor/Private/PhysicsAssetEditorSharedData.h
@@ -5,6 +5,7 @@
 #include "CoreMinimal.h"
 #include "Math/Axis.h"
 #include "PreviewScene.h"
+#include "PhysicsAssetEditorSelection.h"
 #include "PhysicsAssetUtils.h"
 #include "PhysicsEngine/ConstraintInstance.h"
 #include "PhysicsEngine/ShapeElem.h"
@@ -12,16 +13,17 @@
 
 #include "PhysicsAssetEditorSharedData.generated.h"
 
+struct FBoneVertInfo;
+
+class FPhysicsAssetEditorSharedData;
+class IPersonaPreviewScene;
 class UBodySetup;
-class UPhysicsAssetEditorSkeletalMeshComponent;
 class UPhysicsAsset;
-class UPhysicsConstraintTemplate;
 class UPhysicsAssetEditorPhysicsHandleComponent;
+class UPhysicsAssetEditorSkeletalMeshComponent;
+class UPhysicsConstraintTemplate;
 class USkeletalMesh;
 class UStaticMeshComponent;
-struct FBoneVertInfo;
-class IPersonaPreviewScene;
-class FPhysicsAssetEditorSharedData;
 
 /** Scoped object that blocks selection broadcasts until it leaves scope */
 struct FScopedBulkSelection
@@ -63,6 +65,10 @@ public:
 class FPhysicsAssetEditorSharedData
 {
 public:
+
+	using SelectionFilterRange = UPhysicsAssetEditorSelection::FilterRange;
+	using SelectionUniqueRange = UPhysicsAssetEditorSelection::UniqueRange;
+
 	/** Constructor/Destructor */
 	FPhysicsAssetEditorSharedData();
 	virtual ~FPhysicsAssetEditorSharedData();
@@ -75,31 +81,7 @@ public:
 	};
 
 	/** Encapsulates a selected set of bodies or constraints */
-	struct FSelection
-	{
-		int32 Index;
-		EAggCollisionShape::Type PrimitiveType;
-		int32 PrimitiveIndex;
-		FTransform WidgetTM;
-		FTransform ManipulateTM;
-		FVector CoMPosition; 
-		bool bIsDirectBodySelection; // When true, this selection represents a body that has been directly selected, not just one of its primitives.
-
-		FSelection(int32 GivenBodyIndex, EAggCollisionShape::Type GivenPrimitiveType, int32 GivenPrimitiveIndex) 
-			: Index(GivenBodyIndex)
-			, PrimitiveType(GivenPrimitiveType)
-			, PrimitiveIndex(GivenPrimitiveIndex)
-			, WidgetTM(FTransform::Identity)
-			, ManipulateTM(FTransform::Identity)
-			, bIsDirectBodySelection(false)
-		{
-		}
-
-		bool operator==(const FSelection& rhs) const
-		{
-			return Index == rhs.Index && PrimitiveType == rhs.PrimitiveType && PrimitiveIndex == rhs.PrimitiveIndex && bIsDirectBodySelection == rhs.bIsDirectBodySelection;
-		}
-	};
+	using FSelection = FPhysicsAssetEditorSelectedElement;
 
 	/** Initializes members */
 	void Initialize(const TSharedRef<IPersonaPreviewScene>& InPreviewScene);
@@ -122,9 +104,14 @@ public:
 	/** Clear all of the selected constraints */
 	void ClearSelectedConstraints();
 
-	/** Set the selection state of a constraint */
-	void SetSelectedConstraint(int32 ConstraintIndex, bool bSelected);
-	void SetSelectedConstraints(const TArray<int32> ConstraintsIndices, bool bSelected);
+	/** Add or remove a constraint from the current selection */
+	void ModifySelectedConstraints(const int32 ConstraintIndex, const bool bSelected);
+
+	/** Add or remove a collection of constraints from the current selection */
+	void ModifySelectedConstraints(const TArray<int32>& ConstraintsIndices, const bool bSelected);
+
+	/** Set the current selection */
+	void SetSelectedConstraints(const TArray<int32>& ConstraintsIndices);
 
 	/** Check whether the constraint at the specified index is selected */
 	bool IsConstraintSelected(int32 ConstraintIndex) const;
@@ -190,18 +177,43 @@ public:
 	void SetGroupSelectionActive(const bool bIsActive);
 	bool IsGroupSelectionActive() const;
 
-	/** Center of Mass editing */
+	void ModifySelected(const TArray<FSelection>& InSelectedElements, const bool bSelected);
+	void SetSelected(const TArray<FSelection>& InSelectedElements);
+	bool IsSelected(const FSelection& InSelection) const;
+	
+	/** Center of Mass Selection */ 
 	void ClearSelectedCoMs();
-	void SetSelectedCoM(const FSelection& InSelectedCoM, const bool bSelected);
-	void SetSelectedCoMs(const TArray<FSelection>& InSelectedElements, const bool bSelected);
+	void ModifySelectedCoMs(const FSelection& InSelectedCoM, const bool bSelected);
+	void ModifySelectedCoMs(const TArray<FSelection>& InSelectedElements, const bool bSelected);
+	void SetSelectedCoMs(const FSelection& InSelectedElement);
+	void SetSelectedCoMs(const TArray<FSelection>& InSelectedElements);
 	bool IsCoMSelected(const int32 BodyIndex) const;
 
-	/** Collision geometry editing */
+	/** Body Selection */
 	void ClearSelectedBody();
-	void SetSelectedBody(const FSelection& Body, bool bSelected);
-	void SetSelectedBodies(const TArray<FSelection>& Bodies, bool bSelected);
+	void ModifySelectedBodies(const FSelection& Body, bool bSelected);
+	void ModifySelectedBodies(const TArray<FSelection>& Bodies, bool bSelected);
+	void SetSelectedBodies(const FSelection& InSelectedElement);
+	void SetSelectedBodies(const TArray<FSelection>& Bodies);
 	bool IsBodySelected(const FSelection& Body) const;
 	bool IsBodySelected(const int32 BodyIndex) const;
+
+	void ModifySelectedBodies(const int32 BodyIndex, const bool bSelected);
+	void ModifySelectedBodies(const TArray<int32>& BodiesIndices, const bool bSelected);
+	void SetSelectedBodies(const int32 BodyIndex);
+	void SetSelectedBodies(const TArray<int32>& BodiesIndices);
+
+	void SetSelectedBodiesAllPrimitive(const TArray<int32>& BodiesIndices, bool bSelected);
+	void SetSelectedBodiesPrimitivesWithCollisionType(const TArray<int32>& BodiesIndices, const ECollisionEnabled::Type CollisionType, bool bSelected);
+	void SetSelectedBodiesPrimitives(const TArray<int32>& BodiesIndices, bool bSelected, const TFunction<bool(const TArray<FSelection>&, const int32 BodyIndex, const FKShapeElem&)>& Predicate);
+
+	/** Primitives Selection */
+	void ClearSelectedPrimitives();
+	void ModifySelectedPrimitives(const FSelection& InSelectedPrimitive, const bool bSelected);
+	void ModifySelectedPrimitives(const TArray<FSelection>& InSelectedElements, const bool bSelected);
+	void SetSelectedPrimitives(const FSelection& InSelectedElement);
+	void SetSelectedPrimitives(const TArray<FSelection>& InSelectedElements);
+	
 	void ToggleSelectionType(bool bIgnoreUserConstraints = true);
 	void ToggleShowSelected();
 	bool IsBodyHidden(const int32 BodyIndex) const;
@@ -219,12 +231,6 @@ public:
 	void ToggleShowOnlySelected();
 	void ShowSelected();
 	void HideSelected();
-	void SetSelectedBodyAnyPrimitive(int32 BodyIndex, bool bSelected);
-	void SetSelectedBodiesAnyPrimitive(const TArray<int32>& BodiesIndices, bool bSelected);
-	void SetSelectedBodiesAllPrimitive(const TArray<int32>& BodiesIndices, bool bSelected);
-	void SetSelectedBodiesFromDirectlInput(const TArray<int32>& BodiesIndices, const bool bSelected);
-	void SetSelectedBodiesPrimitivesWithCollisionType(const TArray<int32>& BodiesIndices, const ECollisionEnabled::Type CollisionType, bool bSelected);
-	void SetSelectedBodiesPrimitives(const TArray<int32>& BodiesIndices, bool bSelected, const TFunction<bool(const TArray<FSelection>&, const int32 BodyIndex, const FKShapeElem&)>& Predicate);
 	void RefreshPhysicsAssetChange(const UPhysicsAsset* InPhysAsset, bool bFullClothRefresh = true);
 	void MakeNewBody(int32 NewBoneIndex, bool bAutoSelect = true);
 	void MakeNewBody(const FPhysAssetCreateParams& NewBodyData, const int32 NewBoneIndex, const bool bAutoSelect = true);
@@ -329,12 +335,46 @@ public:
 	FVector CalculateCoMNudgeForWorldSpacePosition(const int32 BodyIndex, const FVector& CoMPositionWorldSpace) const;
 
 	// Make a copy of the current component space CoM position from each selected physics body - called before a change in physics body transform
-	void RecordSelectedBodiesCoM();
+	void RecordSelectedCoM();
+
+	const FVector* FindManipulatedBodyCoMPosition(const int32 BodyIndex) const
+	{
+		return ManipulatedBodyCoMPositionMap.Find(BodyIndex);
+	}
+
+	FVector* FindManipulatedBodyCoMPosition(const int32 BodyIndex)
+	{
+		return ManipulatedBodyCoMPositionMap.Find(BodyIndex);
+	}
+
+	void PostManipulationUpdateCoM();
+	void UpdateCoM();
+
+	const UPhysicsAssetEditorSelection* GetSelectedObjects() const;
+
+	SelectionFilterRange SelectedBodies() const;
+	SelectionFilterRange SelectedCoMs() const;
+	SelectionFilterRange SelectedConstraints() const;
+	SelectionFilterRange SelectedPrimitives() const;
+	SelectionFilterRange SelectedBodiesAndPrimitives() const;
+	SelectionUniqueRange UniqueSelectionReferencingBodies() const;
 	
-	// Update the CoM nudge for each selected physics body - called after a change in physics body transform.
-	void UpdateSelectedBodiesCoM();
+	// Returns the most recently selected body or primitive - this is useful as selecting a primitive often acts in the same way as selecting its owning body.
+	const FSelection* GetSelectedBodyOrPrimitive() const;
+	const FSelection* GetSelectedBody() const;
+	const FSelection* GetSelectedCoM() const;
+	const FSelection* GetSelectedConstraint() const;
+	const FSelection* GetSelectedPrimitive() const;
+
+	FVector GetSelectedCoMPosition();
+
+	/** Clears all the selected objects */
+	void ClearSelected();
 
 private:
+
+	bool ModifySelectionInternal(TFunctionRef<bool(void)> SelectionOperation);
+
 	/** Initializes a constraint setup */
 	void InitConstraintSetup(UPhysicsConstraintTemplate* ConstraintSetup, int32 ChildBodyIndex, int32 ParentBodyIndex);
 
@@ -362,12 +402,9 @@ private:
 	/** Generate a new unique name for a constraint */
 	FString MakeUniqueNewConstraintName();
 
-	/** Clears all the selected objects */
-	void ClearSelected();
-
 public:
 	/** Callback for handling selection changes */
-	DECLARE_EVENT_ThreeParams(FPhysicsAssetEditorSharedData, FSelectionChanged, const TArray<FSelection>&, const TArray<FSelection>&, const TArray<FSelection>&);
+	DECLARE_EVENT_OneParam(FPhysicsAssetEditorSharedData, FSelectionChanged, const TArray<FSelection>&);
 	FSelectionChanged SelectionChangedEvent;
 
 	/** Callback for handling changes to the bone/body/constraint hierarchy */
@@ -378,7 +415,6 @@ public:
 	DECLARE_EVENT(FPhysicsAssetEditorSharedData, FHierarchySelectionChangedEvent);
 	FHierarchySelectionChangedEvent HierarchySelectionChangedEvent;
 	
-
 	/** Callback for triggering a refresh of the preview viewport */
 	DECLARE_EVENT(FPhysicsAssetEditorSharedData, FPreviewChanged);
 	FPreviewChanged PreviewChangedEvent;
@@ -410,47 +446,17 @@ public:
 	/** List of bodies that don't collide with the currently selected collision body */
 	TArray<int32> NoCollisionBodies;
 
-	/** Center of Mass editing */
-	TArray<FSelection> SelectedCoMs;
-
 	/** Bone info */
 	TArray<FBoneVertInfo> DominantWeightBoneInfos;
 	TArray<FBoneVertInfo> AnyWeightBoneInfos;
 
-	TArray<FSelection> SelectedBodies;
-
-	FSelection * GetSelectedBody()
-	{
-		int32 Count = SelectedBodies.Num();
-		return Count ? &SelectedBodies[Count - 1] : NULL;
-	}
-
-	/** 
-	 * Returns an array containing only the first FSelection that references a particular body 
-	 * (The SelectedBodies Array includes an FSelction for every selected primitive so can 
-	 * include multiple FSelections with the same body index)
-	 */
-	TArray<FSelection> GetUniqueSelectedBodies() const;
+	UPhysicsAssetEditorSelection* SelectedObjects;
 
-	FSelection* GetSelectedCoM()
-	{
-		int32 Count = SelectedCoMs.Num();
-		return Count ? &SelectedCoMs[Count - 1] : NULL;
-	}
+	void BeginManipulation();
+	void EndManipulation();
+	bool IsManipulating() const { return bManipulating; }
 
-	/** Constraint editing */
-	TArray<FSelection> SelectedConstraints;
-	FSelection * GetSelectedConstraint()
-	{
-		int32 Count = SelectedConstraints.Num();
-		return Count ? &SelectedConstraints[Count - 1] : NULL;
-	}
-
-	const FSelection * GetSelectedConstraint() const
-	{
-		int32 Count = SelectedConstraints.Num();
-		return Count ? &SelectedConstraints[Count - 1] : NULL;
-	}
+	TMap<int32, FVector> ManipulatedBodyCoMPositionMap;
 
 	struct FPhysicsAssetRenderSettings* GetRenderSettings() const;
 
diff --git a/Engine/Source/Editor/PhysicsAssetEditor/Private/PhysicsAssetEditorSkeletalMeshComponent.cpp b/Engine/Source/Editor/PhysicsAssetEditor/Private/PhysicsAssetEditorSkeletalMeshComponent.cpp
--- a/Engine/Source/Editor/PhysicsAssetEditor/Private/PhysicsAssetEditorSkeletalMeshComponent.cpp
+++ b/Engine/Source/Editor/PhysicsAssetEditor/Private/PhysicsAssetEditorSkeletalMeshComponent.cpp
@@ -7,6 +7,8 @@
 #include "Components/SkeletalMeshComponent.h"
 #include "Materials/Material.h"
 #include "Preferences/PhysicsAssetEditorOptions.h"
+#include "SceneManagement.h"
+#include "PhysicsAssetEditorSelection.h"
 #include "PrimitiveDrawingUtils.h"
 #include "PhysicsAssetEditorSharedData.h"
 #include "PhysicsAssetEditorHitProxies.h"
@@ -230,25 +232,9 @@ FTransform UPhysicsAssetEditorSkeletalMeshComponent::GetPrimitiveTransform(const
 	UBodySetup* SharedBodySetup = SharedData->PhysicsAsset->SkeletalBodySetups[BodyIndex];
 	FVector Scale3D(Scale);
 
-	FTransform ManTM = FTransform::Identity;
-
-	if (SharedData->bManipulating && !SharedData->bRunningSimulation)
-	{
-		FPhysicsAssetEditorSharedData::FSelection Body(BodyIndex, PrimType, PrimIndex);
-		for (int32 i = 0; i<SharedData->SelectedBodies.Num(); ++i)
-		{
-			if (Body == SharedData->SelectedBodies[i])
-			{
-				ManTM = SharedData->SelectedBodies[i].ManipulateTM;
-				break;
-			}
-
-		}
-	}
-
 	if (const FKShapeElem* const Prim = SharedBodySetup->AggGeom.GetElement(PrimType, PrimIndex))
 	{
-		FTransform PrimTM = ManTM * Prim->GetTransform();
+		FTransform PrimTM = Prim->GetTransform();
 		PrimTM.ScaleTranslation(Scale3D);
 		return PrimTM * BoneTM;
 	}
@@ -276,8 +262,6 @@ FColor UPhysicsAssetEditorSkeletalMeshComponent::GetPrimitiveColor(const int32 B
 		}
 	}
 
-	FPhysicsAssetEditorSharedData::FSelection Body(BodyIndex, PrimitiveType, PrimitiveIndex);
-
 	static FName SelectionColorName(TEXT("SelectionColor"));
 	const FSlateColor SelectionColor = FAppStyle::GetSlateColor(SelectionColorName);
 	const FLinearColor SelectionColorLinear(SelectionColor.IsColorSpecified() ? SelectionColor.GetSpecifiedColor() : FLinearColor::White);
@@ -285,14 +269,14 @@ FColor UPhysicsAssetEditorSkeletalMeshComponent::GetPrimitiveColor(const int32 B
 	const FColor ElemSelectedBodyColor = (SelectionColorLinear* 0.5f).ToFColor(true);
 
 	bool bInBody = false;
-	for (int32 i = 0; i<SharedData->SelectedBodies.Num(); ++i)
+	for (const FPhysicsAssetEditorSharedData::FSelection& SelectedElement : SharedData->SelectedPrimitives())
 	{
-		if (BodyIndex == SharedData->SelectedBodies[i].Index)
+		if (BodyIndex == SelectedElement.Index)
 		{
 			bInBody = true;
 		}
 
-		if (Body == SharedData->SelectedBodies[i] && !SharedData->bRunningSimulation)
+		if (IsReferencingPrimitive(SelectedElement, BodyIndex, PrimitiveType, PrimitiveIndex) && !SharedData->bRunningSimulation)
 		{
 			return ElemSelectedColor;
 		}
@@ -317,7 +301,7 @@ FColor UPhysicsAssetEditorSkeletalMeshComponent::GetPrimitiveColor(const int32 B
 	}
 	else
 	{
-		if (!SharedData->bRunningSimulation && SharedData->SelectedBodies.Num())
+		if (!SharedData->bRunningSimulation && !SharedData->SelectedBodiesAndPrimitives().IsEmpty())
 		{
 			// If there is no collision with this body, use 'no collision material'.
 			if (SharedData->NoCollisionBodies.Find(BodyIndex) != INDEX_NONE)
@@ -337,11 +321,9 @@ UMaterialInterface* UPhysicsAssetEditorSkeletalMeshComponent::GetPrimitiveMateri
 		return PrimitiveType == EAggCollisionShape::TaperedCapsule ? BoneNoCollisionMaterial : BoneUnselectedMaterial;
 	}
 
-	FPhysicsAssetEditorSharedData::FSelection Body(BodyIndex, PrimitiveType, PrimitiveIndex);
-
-	for (int32 i = 0; i < SharedData->SelectedBodies.Num(); ++i)
+	for (const FPhysicsAssetEditorSharedData::FSelection& SelectedElement : SharedData->SelectedBodiesAndPrimitives())
 	{
-		if (Body == SharedData->SelectedBodies[i] && !SharedData->bRunningSimulation)
+		if (IsReferencingPrimitive(SelectedElement, BodyIndex, PrimitiveType, PrimitiveIndex))
 		{
 			return ElemSelectedMaterial;
 		}
@@ -353,7 +335,7 @@ UMaterialInterface* UPhysicsAssetEditorSkeletalMeshComponent::GetPrimitiveMateri
 	}
 
 	// If there is no collision with this body, use 'no collision material'.
-	if (SharedData->SelectedBodies.Num() && SharedData->NoCollisionBodies.Find(BodyIndex) != INDEX_NONE && !SharedData->bRunningSimulation)
+	if (!SharedData->SelectedBodiesAndPrimitives().IsEmpty() && SharedData->NoCollisionBodies.Find(BodyIndex) != INDEX_NONE && !SharedData->bRunningSimulation)
 	{
 		return BoneNoCollisionMaterial;
 	}
diff --git a/Engine/Source/Runtime/Engine/Classes/PhysicsEngine/ShapeElem.h b/Engine/Source/Runtime/Engine/Classes/PhysicsEngine/ShapeElem.h
--- a/Engine/Source/Runtime/Engine/Classes/PhysicsEngine/ShapeElem.h
+++ b/Engine/Source/Runtime/Engine/Classes/PhysicsEngine/ShapeElem.h
@@ -10,6 +10,7 @@
 #include "Engine/EngineTypes.h"
 #include "ShapeElem.generated.h"
 
+UENUM()
 namespace EAggCollisionShape
 {
 	enum Type : int
