[UE] PhAT Register Selection object with GC to avoid unexpected destruction

#rnx

diff --git a/Engine/Source/Editor/PhysicsAssetEditor/Private/PhysicsAssetEditorSharedData.cpp b/Engine/Source/Editor/PhysicsAssetEditor/Private/PhysicsAssetEditorSharedData.cpp
--- a/Engine/Source/Editor/PhysicsAssetEditor/Private/PhysicsAssetEditorSharedData.cpp
+++ b/Engine/Source/Editor/PhysicsAssetEditor/Private/PhysicsAssetEditorSharedData.cpp
@@ -168,6 +168,7 @@ FPhysicsAssetEditorSharedData::FPhysicsAssetEditorSharedData()
 	check(EditorOptions);
 	EditorOptions->LoadConfig();
 
+	// Construct selection manager.
 	SelectedObjects = NewObject<UPhysicsAssetEditorSelection>(GetTransientPackage(), MakeUniqueObjectName(GetTransientPackage(), UPhysicsAssetEditorSelection::StaticClass(), FName(TEXT("PhysicsAssetEditorSelectedObjects"))), RF_Transactional);
 	check(SelectedObjects);
 }
@@ -3541,6 +3542,7 @@ void FPhysicsAssetEditorSharedData::AddReferencedObjects(FReferenceCollector& Co
 	Collector.AddReferencedObject(PhysicalAnimationComponent);
 	Collector.AddReferencedObject(EditorOptions);
 	Collector.AddReferencedObject(MouseHandle);
+	Collector.AddReferencedObject(SelectedObjects);
 
 	if (PreviewScene != nullptr)
 	{
diff --git a/Engine/Source/Editor/PhysicsAssetEditor/Private/PhysicsAssetEditorSharedData.h b/Engine/Source/Editor/PhysicsAssetEditor/Private/PhysicsAssetEditorSharedData.h
--- a/Engine/Source/Editor/PhysicsAssetEditor/Private/PhysicsAssetEditorSharedData.h
+++ b/Engine/Source/Editor/PhysicsAssetEditor/Private/PhysicsAssetEditorSharedData.h
@@ -450,7 +450,7 @@ public:
 	TArray<FBoneVertInfo> DominantWeightBoneInfos;
 	TArray<FBoneVertInfo> AnyWeightBoneInfos;
 
-	UPhysicsAssetEditorSelection* SelectedObjects;
+	TObjectPtr<UPhysicsAssetEditorSelection> SelectedObjects;
 
 	void BeginManipulation();
 	void EndManipulation();
