[UE] PhAT Initialize SelectedElement index

#rnx

diff --git a/Engine/Source/Editor/PhysicsAssetEditor/Private/PhysicsAssetEditorSelection.h b/Engine/Source/Editor/PhysicsAssetEditor/Private/PhysicsAssetEditorSelection.h
--- a/Engine/Source/Editor/PhysicsAssetEditor/Private/PhysicsAssetEditorSelection.h
+++ b/Engine/Source/Editor/PhysicsAssetEditor/Private/PhysicsAssetEditorSelection.h
@@ -41,7 +41,7 @@ struct FPhysicsAssetEditorSelectedElement
 	bool HasType(const uint32 InElementTypeFlags) const;
 
 	UPROPERTY()
-	int32 Index;
+	int32 Index = INDEX_NONE;
 
 	UPROPERTY()
 	TEnumAsByte<EAggCollisionShape::Type> PrimitiveType = EAggCollisionShape::Unknown;
