[UE] PhAT Create Primive and Body Selection element when selecting a body in the tree view

#rnx

diff --git a/Engine/Source/Editor/PhysicsAssetEditor/Private/PhysicsAssetEditor.cpp b/Engine/Source/Editor/PhysicsAssetEditor/Private/PhysicsAssetEditor.cpp
--- a/Engine/Source/Editor/PhysicsAssetEditor/Private/PhysicsAssetEditor.cpp
+++ b/Engine/Source/Editor/PhysicsAssetEditor/Private/PhysicsAssetEditor.cpp
@@ -4431,7 +4431,9 @@ void FPhysicsAssetEditor::HandleSelectionChanged(const TArrayView<TSharedPtr<ISk
 				if (Item->IsOfType<FSkeletonTreePhysicsBodyItem>())
 				{
 					TSharedPtr<FSkeletonTreePhysicsBodyItem> SkeletonTreePhysicsBodyItem = StaticCastSharedPtr<FSkeletonTreePhysicsBodyItem>(Item);
-					SelectedElements.Add(MakeBodySelection(SharedData->PhysicsAsset, SkeletonTreePhysicsBodyItem->GetBodySetupIndex()));
+					const FPhysicsAssetEditorSharedData::FSelection BodySelection = MakeBodySelection(SharedData->PhysicsAsset, SkeletonTreePhysicsBodyItem->GetBodySetupIndex());
+					SelectedElements.Add(BodySelection);
+					SelectedElements.Add(MakePrimitiveSelection(BodySelection.GetIndex(), BodySelection.GetPrimitiveType(), BodySelection.GetPrimitiveIndex()));
 				}
 				else if (Item->IsOfType<FSkeletonTreePhysicsShapeItem>())
 				{
