[UE] PhAT Fix for colliding bodies highlighting in the viewport
- Call UpdateNoCollisionBodies on selection change

#rnx

diff --git a/Engine/Source/Editor/PhysicsAssetEditor/Private/PhysicsAssetEditorSharedData.cpp b/Engine/Source/Editor/PhysicsAssetEditor/Private/PhysicsAssetEditorSharedData.cpp
--- a/Engine/Source/Editor/PhysicsAssetEditor/Private/PhysicsAssetEditorSharedData.cpp
+++ b/Engine/Source/Editor/PhysicsAssetEditor/Private/PhysicsAssetEditorSharedData.cpp
@@ -2502,6 +2502,7 @@ bool FPhysicsAssetEditorSharedData::ModifySelectionInternal(TFunctionRef<bool(vo
 	if (!InsideSelChange && SelectionOperation())
 	{
 		BroadcastSelectionChanged();
+		UpdateNoCollisionBodies();
 
 		++InsideSelChange;
 		BroadcastPreviewChanged();
