[UE] PhAT: Fixes for working with skinned meshes.
- Fix for editor crash when creating a skinned mesh when no other primitives are present.
- Support the new primitive types for selection and automatic naming on creation,
- Support renaming and deleting the full range of primitive types.
- Create skeleton tree elements for new primitive types.
- Support for deleting skinned mesh w/ delete key.
- Fix OOB array access when deleting multiple primitives of the same type from a single body.

#rnx

diff --git a/Engine/Source/Editor/PhysicsAssetEditor/Private/PhysicsAssetEditor.cpp b/Engine/Source/Editor/PhysicsAssetEditor/Private/PhysicsAssetEditor.cpp
--- a/Engine/Source/Editor/PhysicsAssetEditor/Private/PhysicsAssetEditor.cpp
+++ b/Engine/Source/Editor/PhysicsAssetEditor/Private/PhysicsAssetEditor.cpp
@@ -4121,7 +4121,7 @@ void FPhysicsAssetEditor::GenerateSkinnedTriangleMesh()
 		}
 		else
 		{
-			SharedData->AutoNameAllPrimitives(BodyIndex, EFG_SkinnedTriangleMesh);
+			SharedData->AutoNamePrimitive(BodyIndex, EAggCollisionShape::SkinnedTriangleMesh);
 			SharedData->ModifySelectedBodies(BodyIndex, true);
 		}
 	}
diff --git a/Engine/Source/Editor/PhysicsAssetEditor/Private/PhysicsAssetEditorEditMode.cpp b/Engine/Source/Editor/PhysicsAssetEditor/Private/PhysicsAssetEditorEditMode.cpp
--- a/Engine/Source/Editor/PhysicsAssetEditor/Private/PhysicsAssetEditorEditMode.cpp
+++ b/Engine/Source/Editor/PhysicsAssetEditor/Private/PhysicsAssetEditorEditMode.cpp
@@ -842,16 +842,17 @@ bool FPhysicsAssetEditorEditMode::GetCustomInputCoordinateSystem(FMatrix& InMatr
 FVector FPhysicsAssetEditorEditMode::GetWidgetLocation() const
 {
 	// Don't draw widget if nothing selected.
-	if (SharedData->GetSelectedBodyOrPrimitive())
+	if (const FPhysicsAssetEditorSharedData::FSelection* const SelectedObject = SharedData->GetSelectedBodyOrPrimitive())
 	{
-		int32 BoneIndex = SharedData->EditorSkelComp->GetBoneIndex(SharedData->PhysicsAsset->SkeletalBodySetups[SharedData->GetSelectedBodyOrPrimitive()->Index]->BoneName);
-		if(BoneIndex != INDEX_NONE)
+		const int32 BoneIndex = SharedData->EditorSkelComp->GetBoneIndex(SharedData->PhysicsAsset->SkeletalBodySetups[SelectedObject->GetIndex()]->BoneName);
+		check(SelectedObject->GetPrimitiveIndex() != INDEX_NONE); // Expected valid primitive index - all body selections should reference a valid primitive
+		if((BoneIndex != INDEX_NONE) && (SelectedObject->GetPrimitiveIndex() != INDEX_NONE))
 		{
 			FTransform BoneTM = SharedData->EditorSkelComp->GetBoneTransform(BoneIndex);
 			const float Scale = BoneTM.GetScale3D().GetAbsMax();
 			BoneTM.RemoveScaling();
 
-			return SharedData->EditorSkelComp->GetPrimitiveTransform(BoneTM, SharedData->GetSelectedBodyOrPrimitive()->GetIndex(), SharedData->GetSelectedBodyOrPrimitive()->GetPrimitiveType(), SharedData->GetSelectedBodyOrPrimitive()->GetPrimitiveIndex(), Scale).GetTranslation();
+			return SharedData->EditorSkelComp->GetPrimitiveTransform(BoneTM, SelectedObject->GetIndex(), SelectedObject->GetPrimitiveType(), SelectedObject->GetPrimitiveIndex(), Scale).GetTranslation();
 		}
 	}
 	else if (SharedData->GetSelectedCoM())
diff --git a/Engine/Source/Editor/PhysicsAssetEditor/Private/PhysicsAssetEditorSelection.cpp b/Engine/Source/Editor/PhysicsAssetEditor/Private/PhysicsAssetEditorSelection.cpp
--- a/Engine/Source/Editor/PhysicsAssetEditor/Private/PhysicsAssetEditorSelection.cpp
+++ b/Engine/Source/Editor/PhysicsAssetEditor/Private/PhysicsAssetEditorSelection.cpp
@@ -103,6 +103,10 @@ void InitializeSelectionWithFirstPrimitive(FPhysicsAssetEditorSelectedElement& S
 		InitializeSelectionWithFirstPrimitive(Selection, AggGeom.SphylElems);
 		InitializeSelectionWithFirstPrimitive(Selection, AggGeom.ConvexElems);
 		InitializeSelectionWithFirstPrimitive(Selection, AggGeom.TaperedCapsuleElems);
+		InitializeSelectionWithFirstPrimitive(Selection, AggGeom.LevelSetElems);
+		InitializeSelectionWithFirstPrimitive(Selection, AggGeom.SkinnedLevelSetElems);
+		InitializeSelectionWithFirstPrimitive(Selection, AggGeom.MLLevelSetElems);
+		InitializeSelectionWithFirstPrimitive(Selection, AggGeom.SkinnedTriangleMeshElems);
 	}
 }
 
diff --git a/Engine/Source/Editor/PhysicsAssetEditor/Private/PhysicsAssetEditorSharedData.cpp b/Engine/Source/Editor/PhysicsAssetEditor/Private/PhysicsAssetEditorSharedData.cpp
--- a/Engine/Source/Editor/PhysicsAssetEditor/Private/PhysicsAssetEditorSharedData.cpp
+++ b/Engine/Source/Editor/PhysicsAssetEditor/Private/PhysicsAssetEditorSharedData.cpp
@@ -1787,6 +1787,9 @@ EAggCollisionShape::Type ConvertPhysicsAssetGeomTypeToAggCollisionShapeType(EPhy
 	case EPhysAssetFitGeomType::EFG_SingleConvexHull:	return EAggCollisionShape::Type::Convex;
 	case EPhysAssetFitGeomType::EFG_MultiConvexHull:	return EAggCollisionShape::Type::Convex;
 	case EPhysAssetFitGeomType::EFG_LevelSet:			return EAggCollisionShape::Type::LevelSet;
+	case EPhysAssetFitGeomType::EFG_SkinnedLevelSet:	return EAggCollisionShape::Type::SkinnedLevelSet;
+	case EPhysAssetFitGeomType::EFG_MLLevelSet:			return EAggCollisionShape::Type::MLLevelSet;
+	case EPhysAssetFitGeomType::EFG_SkinnedTriangleMesh:return EAggCollisionShape::Type::SkinnedTriangleMesh;
 	default:											return EAggCollisionShape::Type::Unknown;
 	}
 }
@@ -1801,6 +1804,9 @@ EPhysAssetFitGeomType ConvertAggCollisionShapeTypeToPhysicsAssetGeomType(const E
 	case EAggCollisionShape::Type::TaperedCapsule:		return EPhysAssetFitGeomType::EFG_TaperedCapsule;
 	case EAggCollisionShape::Type::Convex:				return EPhysAssetFitGeomType::EFG_SingleConvexHull;
 	case EAggCollisionShape::Type::LevelSet:			return EPhysAssetFitGeomType::EFG_LevelSet;
+	case EAggCollisionShape::Type::SkinnedLevelSet:		return EPhysAssetFitGeomType::EFG_SkinnedLevelSet;
+	case EAggCollisionShape::Type::MLLevelSet:			return EPhysAssetFitGeomType::EFG_MLLevelSet;
+	case EAggCollisionShape::Type::SkinnedTriangleMesh:	return EPhysAssetFitGeomType::EFG_SkinnedTriangleMesh;
 	default:											return EPhysAssetFitGeomType(INDEX_NONE);
 	}
 }
@@ -1838,6 +1844,18 @@ void FPhysicsAssetEditorSharedData::AutoNameAllPrimitives(int32 BodyIndex, EAggC
 		case EAggCollisionShape::TaperedCapsule:
 			PrimitiveCount = BodySetup->AggGeom.TaperedCapsuleElems.Num();
 			break;
+		case EAggCollisionShape::LevelSet:
+			PrimitiveCount = BodySetup->AggGeom.LevelSetElems.Num();
+			break;
+		case EAggCollisionShape::SkinnedLevelSet:
+			PrimitiveCount = BodySetup->AggGeom.SkinnedLevelSetElems.Num();
+			break;
+		case EAggCollisionShape::MLLevelSet:
+			PrimitiveCount = BodySetup->AggGeom.MLLevelSetElems.Num();
+			break;
+		case EAggCollisionShape::SkinnedTriangleMesh:
+			PrimitiveCount = BodySetup->AggGeom.SkinnedTriangleMeshElems.Num();
+			break;
 		}
 
 		for (int32 PrimitiveIndex = 0; PrimitiveIndex < PrimitiveCount; PrimitiveIndex++)
@@ -1856,65 +1874,54 @@ void FPhysicsAssetEditorSharedData::AutoNamePrimitive(int32 BodyIndex, EAggColli
 
 	if (UBodySetup* BodySetup = PhysicsAsset->SkeletalBodySetups[BodyIndex])
 	{
+		auto SetElementName = [BodySetup, PrimitiveType, &PrimitiveIndex](auto& GeometryCollection, const TCHAR* NamePostfixText)
+			{
+				if (PrimitiveIndex == INDEX_NONE)
+				{
+					PrimitiveIndex = GeometryCollection.Num() - 1;
+				}
+				if (GeometryCollection.IsValidIndex(PrimitiveIndex))
+				{
+					FName PrimitiveName(FString::Printf(TEXT("%s_%s"), *BodySetup->BoneName.ToString(), NamePostfixText));
+					GeometryCollection[PrimitiveIndex].SetName(PrimitiveName);
+				}
+			};
+
 		if (PrimitiveType == EAggCollisionShape::Sphere)
 		{
-			if (PrimitiveIndex == INDEX_NONE)
-			{
-				PrimitiveIndex = BodySetup->AggGeom.SphereElems.Num() - 1;
-			}
-			if (BodySetup->AggGeom.SphereElems.IsValidIndex(PrimitiveIndex))
-			{
-				FName PrimitiveName(FString::Printf(TEXT("%s_sphere"), *BodySetup->BoneName.ToString()));
-				BodySetup->AggGeom.SphereElems[PrimitiveIndex].SetName(PrimitiveName);
-			}
+			SetElementName(BodySetup->AggGeom.SphereElems, TEXT("sphere"));
 		}
 		else if (PrimitiveType == EAggCollisionShape::Box)
 		{
-			if (PrimitiveIndex == INDEX_NONE)
-			{
-				PrimitiveIndex = BodySetup->AggGeom.BoxElems.Num() - 1;
-			}
-			if (BodySetup->AggGeom.BoxElems.IsValidIndex(PrimitiveIndex))
-			{
-				FName PrimitiveName(FString::Printf(TEXT("%s_box"), *BodySetup->BoneName.ToString()));
-				BodySetup->AggGeom.BoxElems[PrimitiveIndex].SetName(PrimitiveName);
-			}
+			SetElementName(BodySetup->AggGeom.BoxElems, TEXT("box"));
 		}
 		else if (PrimitiveType == EAggCollisionShape::Sphyl)
 		{
-			if (PrimitiveIndex == INDEX_NONE)
-			{
-				PrimitiveIndex = BodySetup->AggGeom.SphylElems.Num() - 1;
-			}
-			if (BodySetup->AggGeom.SphylElems.IsValidIndex(PrimitiveIndex))
-			{
-				FName PrimitiveName(FString::Printf(TEXT("%s_capsule"), *BodySetup->BoneName.ToString()));
-				BodySetup->AggGeom.SphylElems[PrimitiveIndex].SetName(PrimitiveName);
-			}
+			SetElementName(BodySetup->AggGeom.SphylElems, TEXT("capsule"));
 		}
 		else if (PrimitiveType == EAggCollisionShape::Convex)
 		{
-			if (PrimitiveIndex == INDEX_NONE)
-			{
-				PrimitiveIndex = BodySetup->AggGeom.ConvexElems.Num() - 1;
-			}
-			if (BodySetup->AggGeom.ConvexElems.IsValidIndex(PrimitiveIndex))
-			{
-				FName PrimitiveName(FString::Printf(TEXT("%s_convex"), *BodySetup->BoneName.ToString()));
-				BodySetup->AggGeom.ConvexElems[PrimitiveIndex].SetName(PrimitiveName);
-			}
+			SetElementName(BodySetup->AggGeom.ConvexElems, TEXT("convex"));
 		}
 		else if (PrimitiveType == EAggCollisionShape::TaperedCapsule)
 		{
-			if (PrimitiveIndex == INDEX_NONE)
-			{
-				PrimitiveIndex = BodySetup->AggGeom.TaperedCapsuleElems.Num() - 1;
-			}
-			if (BodySetup->AggGeom.TaperedCapsuleElems.IsValidIndex(PrimitiveIndex))
-			{
-				FName PrimitiveName(FString::Printf(TEXT("%s_tapered_capsule"), *BodySetup->BoneName.ToString()));
-				BodySetup->AggGeom.TaperedCapsuleElems[PrimitiveIndex].SetName(PrimitiveName);
-			}
+			SetElementName(BodySetup->AggGeom.TaperedCapsuleElems, TEXT("tapered_capsule"));
+		}
+		else if (PrimitiveType == EAggCollisionShape::LevelSet)
+		{
+			SetElementName(BodySetup->AggGeom.LevelSetElems, TEXT("level_set"));
+		}
+		else if (PrimitiveType == EAggCollisionShape::SkinnedLevelSet)
+		{
+			SetElementName(BodySetup->AggGeom.SkinnedLevelSetElems, TEXT("skinned_level_set"));
+		}
+		else if (PrimitiveType == EAggCollisionShape::MLLevelSet)
+		{
+			SetElementName(BodySetup->AggGeom.MLLevelSetElems, TEXT("ml_level_set"));
+		}
+		else if (PrimitiveType == EAggCollisionShape::SkinnedTriangleMesh)
+		{
+			SetElementName(BodySetup->AggGeom.SkinnedTriangleMeshElems, TEXT("skinned_triangle_mesh"));
 		}
 	}
 }
@@ -3104,16 +3111,23 @@ void FPhysicsAssetEditorSharedData::DeleteCurrentPrim()
 
 	const FScopedTransaction Transaction( NSLOCTEXT("UnrealEd", "DeletePrimitive", "Delete Primitive") );
 
-	for (TMap<UBodySetup*, TArray<FSelection> >::TConstIterator It(BodySelectionMap); It; ++It)
+	for (TMap<UBodySetup*, TArray<FSelection> >::TIterator It(BodySelectionMap); It; ++It)
 	{
 		UBodySetup * BodySetup = It.Key();
-		const TArray<FSelection> & SelectedPrimitives = It.Value();
+		TArray<FSelection>& SelectedPrimitives = It.Value();
+
+		// Sort selected primitives by primitive index to ensure we update element indexes correctly as we modify the geometry arrays.
+		SelectedPrimitives.Sort([](const FSelection& LhsElement, const FSelection& RhsElement) -> bool { return LhsElement.GetPrimitiveIndex() < RhsElement.GetPrimitiveIndex(); });
 
 		int32 SphereDeletedCount = 0;
 		int32 BoxDeletedCount = 0;
 		int32 SphylDeletedCount = 0;
 		int32 ConvexDeletedCount = 0;
 		int32 TaperedCapsuleDeletedCount = 0;
+		int32 LevelSetDeletedCount = 0;
+		int32 SkinnedLevelSetDeletedCount = 0;
+		int32 MLLevelSetDeletedCount = 0;
+		int32 SkinnedTriangleMeshDeletedCount = 0;
 
 		for (int32 i = 0; i < SelectedPrimitives.Num(); ++i)
 		{
@@ -3144,6 +3158,22 @@ void FPhysicsAssetEditorSharedData::DeleteCurrentPrim()
 			{
 				BodySetup->AggGeom.TaperedCapsuleElems.RemoveAt(SelectedBody.PrimitiveIndex - (TaperedCapsuleDeletedCount++));
 			}
+			else if (SelectedBody.GetPrimitiveType() == EAggCollisionShape::LevelSet)
+			{
+				BodySetup->AggGeom.LevelSetElems.RemoveAt(SelectedBody.PrimitiveIndex - (LevelSetDeletedCount++));
+			}
+			else if (SelectedBody.GetPrimitiveType() == EAggCollisionShape::SkinnedLevelSet)
+			{
+				BodySetup->AggGeom.SkinnedLevelSetElems.RemoveAt(SelectedBody.PrimitiveIndex - (SkinnedLevelSetDeletedCount++));
+			}
+			else if (SelectedBody.GetPrimitiveType() == EAggCollisionShape::MLLevelSet)
+			{
+				BodySetup->AggGeom.MLLevelSetElems.RemoveAt(SelectedBody.PrimitiveIndex - (MLLevelSetDeletedCount++));
+			}
+			else if (SelectedBody.GetPrimitiveType() == EAggCollisionShape::SkinnedTriangleMesh)
+			{
+				BodySetup->AggGeom.SkinnedTriangleMeshElems.RemoveAt(SelectedBody.PrimitiveIndex - (SkinnedTriangleMeshDeletedCount++));
+			}
 
 			// If this bone has no more geometry - remove it totally.
 			if (BodySetup->AggGeom.GetElementCount() == 0)
diff --git a/Engine/Source/Editor/PhysicsAssetEditor/Private/PhysicsAssetEditorSkeletonTreeBuilder.cpp b/Engine/Source/Editor/PhysicsAssetEditor/Private/PhysicsAssetEditorSkeletonTreeBuilder.cpp
--- a/Engine/Source/Editor/PhysicsAssetEditor/Private/PhysicsAssetEditorSkeletonTreeBuilder.cpp
+++ b/Engine/Source/Editor/PhysicsAssetEditor/Private/PhysicsAssetEditorSkeletonTreeBuilder.cpp
@@ -185,6 +185,26 @@ void FPhysicsAssetEditorSkeletonTreeBuilder::AddBodies(FSkeletonTreeBuilderOutpu
 							{
 								Output.Add(MakeShared<FSkeletonTreePhysicsShapeItem>(BodySetup, BoneName, BodySetupIndex, EAggCollisionShape::TaperedCapsule, ShapeIndex, SkeletonTreePtr.Pin().ToSharedRef()), BoneName, FSkeletonTreePhysicsBodyItem::GetTypeId());
 							}
+
+							for (ShapeIndex = 0; ShapeIndex < AggGeom.LevelSetElems.Num(); ++ShapeIndex)
+							{
+								Output.Add(MakeShared<FSkeletonTreePhysicsShapeItem>(BodySetup, BoneName, BodySetupIndex, EAggCollisionShape::LevelSet, ShapeIndex, SkeletonTreePtr.Pin().ToSharedRef()), BoneName, FSkeletonTreePhysicsBodyItem::GetTypeId());
+							}
+
+							for (ShapeIndex = 0; ShapeIndex < AggGeom.SkinnedLevelSetElems.Num(); ++ShapeIndex)
+							{
+								Output.Add(MakeShared<FSkeletonTreePhysicsShapeItem>(BodySetup, BoneName, BodySetupIndex, EAggCollisionShape::SkinnedLevelSet, ShapeIndex, SkeletonTreePtr.Pin().ToSharedRef()), BoneName, FSkeletonTreePhysicsBodyItem::GetTypeId());
+							}
+
+							for (ShapeIndex = 0; ShapeIndex < AggGeom.MLLevelSetElems.Num(); ++ShapeIndex)
+							{
+								Output.Add(MakeShared<FSkeletonTreePhysicsShapeItem>(BodySetup, BoneName, BodySetupIndex, EAggCollisionShape::MLLevelSet, ShapeIndex, SkeletonTreePtr.Pin().ToSharedRef()), BoneName, FSkeletonTreePhysicsBodyItem::GetTypeId());
+							}
+
+							for (ShapeIndex = 0; ShapeIndex < AggGeom.SkinnedTriangleMeshElems.Num(); ++ShapeIndex)
+							{
+								Output.Add(MakeShared<FSkeletonTreePhysicsShapeItem>(BodySetup, BoneName, BodySetupIndex, EAggCollisionShape::SkinnedTriangleMesh, ShapeIndex, SkeletonTreePtr.Pin().ToSharedRef()), BoneName, FSkeletonTreePhysicsBodyItem::GetTypeId());
+							}
 						}
 
 						// add constraints for this bone
diff --git a/Engine/Source/Editor/PhysicsAssetEditor/Private/SkeletonTreePhysicsShapeItem.cpp b/Engine/Source/Editor/PhysicsAssetEditor/Private/SkeletonTreePhysicsShapeItem.cpp
--- a/Engine/Source/Editor/PhysicsAssetEditor/Private/SkeletonTreePhysicsShapeItem.cpp
+++ b/Engine/Source/Editor/PhysicsAssetEditor/Private/SkeletonTreePhysicsShapeItem.cpp
@@ -39,6 +39,23 @@ FSkeletonTreePhysicsShapeItem::FSkeletonTreePhysicsShapeItem(USkeletalBodySetup*
 		ShapeBrush = FAppStyle::GetBrush("PhysicsAssetEditor.Tree.TaperedCapsule");
 		DefaultLabel = *FText::Format(LOCTEXT("TaperedCapsuleLabel", "{0} Tapered Capsule {1}"), FText::FromName(InBoneName), FText::AsNumber(ShapeIndex)).ToString();
 		break;
+	case EAggCollisionShape::LevelSet:
+		ShapeBrush = FAppStyle::GetBrush("PhysicsAssetEditor.Tree.Box");
+		DefaultLabel = *FText::Format(LOCTEXT("LevelSetLabel", "{0} Level Set {1}"), FText::FromName(InBoneName), FText::AsNumber(ShapeIndex)).ToString();
+		break;
+	case EAggCollisionShape::SkinnedLevelSet:
+		ShapeBrush = FAppStyle::GetBrush("PhysicsAssetEditor.Tree.Box");
+		DefaultLabel = *FText::Format(LOCTEXT("SkinnedLevelSetLabel", "{0} Skinned Level Set {1}"), FText::FromName(InBoneName), FText::AsNumber(ShapeIndex)).ToString();
+		break;
+	case EAggCollisionShape::MLLevelSet:
+		ShapeBrush = FAppStyle::GetBrush("PhysicsAssetEditor.Tree.Box");
+		DefaultLabel = *FText::Format(LOCTEXT("MLLevelSetLabel", "{0} ML Level Set {1}"), FText::FromName(InBoneName), FText::AsNumber(ShapeIndex)).ToString();
+		break;
+	case EAggCollisionShape::SkinnedTriangleMesh:
+		ShapeBrush = FAppStyle::GetBrush("PhysicsAssetEditor.Tree.Box");
+		DefaultLabel = *FText::Format(LOCTEXT("SkinnedTriangleMeshLabel", "{0} Skinned Triangle Mesh {1}"), FText::FromName(InBoneName), FText::AsNumber(ShapeIndex)).ToString();
+		break;
+
 	default:
 		check(false);
 		break;
@@ -110,37 +127,42 @@ FString FSkeletonTreePhysicsShapeItem::GetNameAsString() const
 {
 	FString StringName;
 
+	auto GetPrimitiveName = [this, &StringName](const auto& GeometryCollection)
+		{
+			if (GeometryCollection.IsValidIndex(ShapeIndex))
+			{
+				StringName = GeometryCollection[ShapeIndex].GetName().GetPlainNameString();
+			}
+		};
+
 	switch (ShapeType)
 	{
 	case EAggCollisionShape::Sphere:
-		if(BodySetup->AggGeom.SphereElems.IsValidIndex(ShapeIndex))
-		{
-			StringName = BodySetup->AggGeom.SphereElems[ShapeIndex].GetName().GetPlainNameString();
-		}
+		GetPrimitiveName(BodySetup->AggGeom.SphereElems);
 		break;
 	case EAggCollisionShape::Box:
-		if(BodySetup->AggGeom.BoxElems.IsValidIndex(ShapeIndex))
-		{
-			StringName = BodySetup->AggGeom.BoxElems[ShapeIndex].GetName().GetPlainNameString();
-		}
+		GetPrimitiveName(BodySetup->AggGeom.BoxElems);
 		break;
 	case EAggCollisionShape::Sphyl:
-		if(BodySetup->AggGeom.SphylElems.IsValidIndex(ShapeIndex))
-		{
-			StringName = BodySetup->AggGeom.SphylElems[ShapeIndex].GetName().GetPlainNameString();
-		}
+		GetPrimitiveName(BodySetup->AggGeom.SphylElems);
 		break;
 	case EAggCollisionShape::Convex:
-		if(BodySetup->AggGeom.ConvexElems.IsValidIndex(ShapeIndex))
-		{
-			StringName = BodySetup->AggGeom.ConvexElems[ShapeIndex].GetName().GetPlainNameString();
-		}
+		GetPrimitiveName(BodySetup->AggGeom.ConvexElems);
 		break;
 	case EAggCollisionShape::TaperedCapsule:
-		if(BodySetup->AggGeom.TaperedCapsuleElems.IsValidIndex(ShapeIndex))
-		{
-			StringName = BodySetup->AggGeom.TaperedCapsuleElems[ShapeIndex].GetName().GetPlainNameString();
-		}
+		GetPrimitiveName(BodySetup->AggGeom.TaperedCapsuleElems);
+		break;
+	case EAggCollisionShape::LevelSet:
+		GetPrimitiveName(BodySetup->AggGeom.LevelSetElems);
+		break;
+	case EAggCollisionShape::SkinnedLevelSet:
+		GetPrimitiveName(BodySetup->AggGeom.SkinnedLevelSetElems);
+		break;
+	case EAggCollisionShape::MLLevelSet:
+		GetPrimitiveName(BodySetup->AggGeom.MLLevelSetElems);
+		break;
+	case EAggCollisionShape::SkinnedTriangleMesh:
+		GetPrimitiveName(BodySetup->AggGeom.SkinnedTriangleMeshElems);
 		break;
 	}
 
@@ -165,37 +187,42 @@ void FSkeletonTreePhysicsShapeItem::HandleTextCommitted(const FText& InText, ETe
 
 		BodySetup->Modify();
 
+		auto SetPrimitiveName = [this, InText](auto& GeometryCollection)
+			{
+				if (GeometryCollection.IsValidIndex(ShapeIndex))
+				{
+					GeometryCollection[ShapeIndex].SetName(*InText.ToString());
+				}
+			};
+
 		switch (ShapeType)
 		{
 		case EAggCollisionShape::Sphere:
-			if(BodySetup->AggGeom.SphereElems.IsValidIndex(ShapeIndex))
-			{
-				BodySetup->AggGeom.SphereElems[ShapeIndex].SetName(*InText.ToString());
-			}
+			SetPrimitiveName(BodySetup->AggGeom.SphereElems);
 			break;
 		case EAggCollisionShape::Box:
-			if(BodySetup->AggGeom.BoxElems.IsValidIndex(ShapeIndex))
-			{
-				BodySetup->AggGeom.BoxElems[ShapeIndex].SetName(*InText.ToString());
-			}
+			SetPrimitiveName(BodySetup->AggGeom.BoxElems);
 			break;
 		case EAggCollisionShape::Sphyl:
-			if(BodySetup->AggGeom.SphylElems.IsValidIndex(ShapeIndex))
-			{
-				BodySetup->AggGeom.SphylElems[ShapeIndex].SetName(*InText.ToString());
-			}
+			SetPrimitiveName(BodySetup->AggGeom.SphylElems);
 			break;
 		case EAggCollisionShape::Convex:
-			if(BodySetup->AggGeom.ConvexElems.IsValidIndex(ShapeIndex))
-			{
-				BodySetup->AggGeom.ConvexElems[ShapeIndex].SetName(*InText.ToString());
-			}
+			SetPrimitiveName(BodySetup->AggGeom.ConvexElems);
 			break;
 		case EAggCollisionShape::TaperedCapsule:
-			if(BodySetup->AggGeom.TaperedCapsuleElems.IsValidIndex(ShapeIndex))
-			{
-				BodySetup->AggGeom.TaperedCapsuleElems[ShapeIndex].SetName(*InText.ToString());
-			}
+			SetPrimitiveName(BodySetup->AggGeom.TaperedCapsuleElems);
+			break;
+		case EAggCollisionShape::LevelSet:
+			SetPrimitiveName(BodySetup->AggGeom.LevelSetElems);
+			break;
+		case EAggCollisionShape::SkinnedLevelSet:
+			SetPrimitiveName(BodySetup->AggGeom.SkinnedLevelSetElems);
+			break;
+		case EAggCollisionShape::MLLevelSet:
+			SetPrimitiveName(BodySetup->AggGeom.MLLevelSetElems);
+			break;
+		case EAggCollisionShape::SkinnedTriangleMesh:
+			SetPrimitiveName(BodySetup->AggGeom.SkinnedTriangleMeshElems);
 			break;
 		default:
 			check(false);
