[UE] PhAT : Fix cross constraints creation between bodies selected in the viewport

diff --git a/Engine/Source/Editor/PhysicsAssetEditor/Private/PhysicsAssetEditor.cpp b/Engine/Source/Editor/PhysicsAssetEditor/Private/PhysicsAssetEditor.cpp
--- a/Engine/Source/Editor/PhysicsAssetEditor/Private/PhysicsAssetEditor.cpp
+++ b/Engine/Source/Editor/PhysicsAssetEditor/Private/PhysicsAssetEditor.cpp
@@ -3350,7 +3350,7 @@ bool FPhysicsAssetEditor::CanDuplicatePrimitive() const
 
 void FPhysicsAssetEditor::OnConstrainChildBodiesToParentBody()
 {
-	const FPhysicsAssetEditorSharedData::FSelection* const LastSelectedBody = SharedData->GetSelectedBody();
+	const FPhysicsAssetEditorSharedData::FSelection* const LastSelectedBody = SharedData->GetSelectedBodyOrPrimitive();
 
 	if (LastSelectedBody && (SharedData->UniqueSelectionReferencingBodies().Num() > 1))
 	{
