[UE] Phat: Fix for Isolate Selected on bodies with multiple primitives

#rnx

diff --git a/Engine/Source/Editor/PhysicsAssetEditor/Private/PhysicsAssetEditorSharedData.cpp b/Engine/Source/Editor/PhysicsAssetEditor/Private/PhysicsAssetEditorSharedData.cpp
--- a/Engine/Source/Editor/PhysicsAssetEditor/Private/PhysicsAssetEditorSharedData.cpp
+++ b/Engine/Source/Editor/PhysicsAssetEditor/Private/PhysicsAssetEditorSharedData.cpp
@@ -1148,7 +1148,7 @@ bool FPhysicsAssetEditorSharedData::IsBodySelected(const int32 BodyIndex) const
 {
 	return Algo::FindByPredicate(SelectedObjects->SelectedElements(), [BodyIndex](const FSelection& Element) 
 		{ 
-			return Element.HasType(FPhysicsAssetEditorSelectedElement::Body) && (Element.Index == BodyIndex); 
+			return Element.HasType(FPhysicsAssetEditorSelectedElement::Body | FPhysicsAssetEditorSelectedElement::Primitive) && (Element.Index == BodyIndex);
 		}) != nullptr;
 }
 
