[UE] PhAT Highlight Overlapping Bodies and rationalize body highlight colors
- Physics bodies can now optionally be drawn in red if they are overlapping with other bodies that they can collide with (i.e. not flagged as no collision) to assist in creating stable physics assets
- Updated highlight colors for selected / unselected / colliding physics bodies to make them easier to understand

#virtualized

diff --git a/Engine/Source/Editor/Persona/Private/PhysicsAssetRenderUtils.cpp b/Engine/Source/Editor/Persona/Private/PhysicsAssetRenderUtils.cpp
--- a/Engine/Source/Editor/Persona/Private/PhysicsAssetRenderUtils.cpp
+++ b/Engine/Source/Editor/Persona/Private/PhysicsAssetRenderUtils.cpp
@@ -72,6 +72,7 @@ FPhysicsAssetRenderSettings::FPhysicsAssetRenderSettings()
 	, bRenderOnlySelectedConstraints(false)
 	, bShowCOM_DEPRECATED(false)
 	, bShowConstraintsAsPoints(false)
+	, bHighlightOverlapingBodies(false)
 	, bDrawViolatedLimits(false)
 	, bHideCenterOfMassForKinematicBodies(true)
 	, BoneUnselectedColor(170, 155, 225)
diff --git a/Engine/Source/Editor/Persona/Public/PhysicsAssetRenderUtils.h b/Engine/Source/Editor/Persona/Public/PhysicsAssetRenderUtils.h
--- a/Engine/Source/Editor/Persona/Public/PhysicsAssetRenderUtils.h
+++ b/Engine/Source/Editor/Persona/Public/PhysicsAssetRenderUtils.h
@@ -116,6 +116,9 @@ public:
 	UPROPERTY()
 	bool bShowConstraintsAsPoints;
 
+	UPROPERTY()
+	bool bHighlightOverlapingBodies;
+
 	UPROPERTY()
 	bool bDrawViolatedLimits;
 	
diff --git a/Engine/Source/Editor/PhysicsAssetEditor/Private/PhysicsAssetEditor.cpp b/Engine/Source/Editor/PhysicsAssetEditor/Private/PhysicsAssetEditor.cpp
--- a/Engine/Source/Editor/PhysicsAssetEditor/Private/PhysicsAssetEditor.cpp
+++ b/Engine/Source/Editor/PhysicsAssetEditor/Private/PhysicsAssetEditor.cpp
@@ -228,6 +228,8 @@ void FPhysicsAssetEditor::InitPhysicsAssetEditor(const EToolkitMode::Type Mode,
 
 	PersonaModule.RecordAssetOpened(FAssetData(ObjectToEdit));
 
+	SharedData->InitializeOverlappingBodyPairs();
+
 	FSkeletonTreeArgs SkeletonTreeArgs;
 	SkeletonTreeArgs.OnSelectionChanged = FOnSkeletonTreeSelectionChanged::CreateSP(this, &FPhysicsAssetEditor::HandleSelectionChanged);
 	SkeletonTreeArgs.PreviewScene = PersonaToolkit->GetPreviewScene();
@@ -1002,6 +1004,7 @@ void FPhysicsAssetEditor::ExtendViewportMenus()
 										Section.AddMenuEntry(Commands.CollisionRenderingMode_Wireframe);
 										Section.AddMenuEntry(Commands.CollisionRenderingMode_SolidWireframe);
 										Section.AddMenuEntry(Commands.CollisionRenderingMode_None);
+										Section.AddMenuEntry(Commands.HighlightOverlappingBodies);
 									}
 
 									{
@@ -1764,6 +1767,12 @@ void FPhysicsAssetEditor::BindCommands()
 		FCanExecuteAction(),
 		FIsActionChecked::CreateSP(this, &FPhysicsAssetEditor::IsHidingKinematicBodies));
 	
+	ViewportCommandList->MapAction(
+		Commands.HighlightOverlappingBodies,
+		FExecuteAction::CreateSP(this, &FPhysicsAssetEditor::ToggleHighlightOverlappingBodies),
+		FCanExecuteAction(),
+		FIsActionChecked::CreateSP(this, &FPhysicsAssetEditor::IsHighlightOverlappingBodies));
+
 	ViewportCommandList->MapAction(
 		Commands.DrawBodyMass,
 		FExecuteAction::CreateSP(this, &FPhysicsAssetEditor::ToggleHideBodyMass),
@@ -2486,6 +2495,8 @@ void FPhysicsAssetEditor::AddNewPrimitive(EAggCollisionShape::Type InPrimitiveTy
 			{
 				check(0);  //unrecognized primitive type
 			}
+
+			SharedData->UpdateOverlappingBodyPairs(BodyIndex);
 		}
 	} // ScopedTransaction
 
@@ -3193,6 +3204,11 @@ void FPhysicsAssetEditor::ToggleHideKinematicBodies()
 	SharedData->EditorOptions->SaveConfig();
 }
 
+void FPhysicsAssetEditor::ToggleHighlightOverlappingBodies()
+{
+	SharedData->ToggleHighlightOverlapingBodies();
+}
+
 void FPhysicsAssetEditor::ToggleHideBodyMass()
 {
 	SharedData->EditorOptions->bHideBodyMass = !SharedData->EditorOptions->bHideBodyMass;
@@ -3214,6 +3230,11 @@ bool FPhysicsAssetEditor::IsHidingKinematicBodies() const
 	return SharedData->EditorOptions->bHideKinematicBodies;
 }
 
+bool FPhysicsAssetEditor::IsHighlightOverlappingBodies() const
+{
+	return SharedData->IsHighlightingOverlapingBodies();
+}
+
 bool FPhysicsAssetEditor::IsHidingBodyMass() const
 {
 	return SharedData->EditorOptions->bHideBodyMass;
diff --git a/Engine/Source/Editor/PhysicsAssetEditor/Private/PhysicsAssetEditor.h b/Engine/Source/Editor/PhysicsAssetEditor/Private/PhysicsAssetEditor.h
--- a/Engine/Source/Editor/PhysicsAssetEditor/Private/PhysicsAssetEditor.h
+++ b/Engine/Source/Editor/PhysicsAssetEditor/Private/PhysicsAssetEditor.h
@@ -265,10 +265,12 @@ private:
 	void ToggleRenderOnlySelectedSolid();
 	void ToggleHideSimulatedBodies();
 	void ToggleHideKinematicBodies();
+	void ToggleHighlightOverlappingBodies();
 	void ToggleHideBodyMass();
 	bool IsRenderingOnlySelectedSolid() const;
 	bool IsHidingSimulatedBodies() const;
 	bool IsHidingKinematicBodies() const;
+	bool IsHighlightOverlappingBodies() const;
 	bool IsHidingBodyMass() const;
 	bool IsDrawingBodyMass() const;
 	void OnToggleMassProperties();
diff --git a/Engine/Source/Editor/PhysicsAssetEditor/Private/PhysicsAssetEditorActions.cpp b/Engine/Source/Editor/PhysicsAssetEditor/Private/PhysicsAssetEditorActions.cpp
--- a/Engine/Source/Editor/PhysicsAssetEditor/Private/PhysicsAssetEditorActions.cpp
+++ b/Engine/Source/Editor/PhysicsAssetEditor/Private/PhysicsAssetEditorActions.cpp
@@ -127,6 +127,8 @@ void FPhysicsAssetEditorCommands::RegisterCommands()
 	UI_COMMAND(ToggleShowOnlyColliding, "Toggle Isolate Colliding Bodies", "Show only the selected body and the ones that can collide with it, hiding all others or Show all.", EUserInterfaceActionType::Button, FInputChord(EKeys::C));
 	UI_COMMAND(ToggleShowOnlyConstrained, "Toggle Isolate Constrained Bodies", "Show only the selected body and the ones that are joint constrained to it, also works from constraints, hiding all others or Show all.", EUserInterfaceActionType::Button, FInputChord(EKeys::J));
 	UI_COMMAND(ToggleShowOnlySelected, "Toggle Isolate Selected", "Show the selected bodies/constraints, hiding all others, or Show all.", EUserInterfaceActionType::Button, FInputChord(EKeys::G));
+	UI_COMMAND(HighlightOverlappingBodies, "Highlight Overlapping Bodies", "Draw colliding bodies that are currently overlapping each other in red (only supports Box, Capsule and Sphere primitive types in Edit mode)", EUserInterfaceActionType::ToggleButton, FInputChord(EModifierKey::Control, EKeys::C));
+
 	UI_COMMAND(ShowAll, "Show All", "Show all bodies/constraints.", EUserInterfaceActionType::Button, FInputChord(EModifierKey::Shift, EKeys::G));
 	UI_COMMAND(HideAll, "Hide All", "Hide all bodies/constraints.", EUserInterfaceActionType::Button, FInputChord(EModifierKey::Control, EKeys::G));
 	UI_COMMAND(DeselectAll, "Deselect All Bodies", "Deselect All", EUserInterfaceActionType::Button, FInputChord(EKeys::Escape));
diff --git a/Engine/Source/Editor/PhysicsAssetEditor/Private/PhysicsAssetEditorActions.h b/Engine/Source/Editor/PhysicsAssetEditor/Private/PhysicsAssetEditorActions.h
--- a/Engine/Source/Editor/PhysicsAssetEditor/Private/PhysicsAssetEditorActions.h
+++ b/Engine/Source/Editor/PhysicsAssetEditor/Private/PhysicsAssetEditorActions.h
@@ -77,6 +77,7 @@ public:
 	TSharedPtr<FUICommandInfo> RenderOnlySelectedSolid;
 	TSharedPtr<FUICommandInfo> HideSimulatedBodies;
 	TSharedPtr<FUICommandInfo> HideKinematicBodies;
+	TSharedPtr<FUICommandInfo> HighlightOverlappingBodies;
 	TSharedPtr<FUICommandInfo> DrawBodyMass;
 	TSharedPtr<FUICommandInfo> DrawCenterOfMass;
 	TSharedPtr<FUICommandInfo> HideCenterOfMassForKinematicBodies;
diff --git a/Engine/Source/Editor/PhysicsAssetEditor/Private/PhysicsAssetEditorEditMode.cpp b/Engine/Source/Editor/PhysicsAssetEditor/Private/PhysicsAssetEditorEditMode.cpp
--- a/Engine/Source/Editor/PhysicsAssetEditor/Private/PhysicsAssetEditorEditMode.cpp
+++ b/Engine/Source/Editor/PhysicsAssetEditor/Private/PhysicsAssetEditorEditMode.cpp
@@ -634,6 +634,14 @@ void FPhysicsAssetEditorEditMode::Tick(FEditorViewportClient* ViewportClient, fl
 	}
 
 	SharedData->UpdateCoM();
+
+	if (SharedData->bManipulating)
+	{
+		for (const FPhysicsAssetEditorSharedData::FSelection& SelectedBody : SharedData->UniqueSelectionReferencingBodies())
+		{
+			SharedData->UpdateOverlappingBodyPairs(SelectedBody.Index);
+		}
+	}
 }
 
 void FPhysicsAssetEditorEditMode::Render(const FSceneView* View, FViewport* Viewport, FPrimitiveDrawInterface* PDI)
diff --git a/Engine/Source/Editor/PhysicsAssetEditor/Private/PhysicsAssetEditorSharedData.cpp b/Engine/Source/Editor/PhysicsAssetEditor/Private/PhysicsAssetEditorSharedData.cpp
--- a/Engine/Source/Editor/PhysicsAssetEditor/Private/PhysicsAssetEditorSharedData.cpp
+++ b/Engine/Source/Editor/PhysicsAssetEditor/Private/PhysicsAssetEditorSharedData.cpp
@@ -2,6 +2,10 @@
 
 #include "PhysicsAssetEditorSharedData.h"
 #include "Animation/MirrorDataTable.h"
+#include "Chaos/GeometryQueries.h"
+#include "Chaos/Box.h"
+#include "Chaos/Capsule.h"
+#include "Chaos/Sphere.h"
 #include "PhysicsAssetEditorPhysicsHandleComponent.h"
 #include "PhysicsAssetRenderUtils.h"
 #include "PhysicsAssetEditorSelection.h"
@@ -93,6 +97,10 @@ TArray<FPhysicsAssetEditorSharedData::FSelection> CreateBodyPrimitivesSelection(
 		SetSelectedBodiesPrimitivesHelper(BodyIndex, AggGeom.SphylElems, NewSelection, Predicate);
 		SetSelectedBodiesPrimitivesHelper(BodyIndex, AggGeom.ConvexElems, NewSelection, Predicate);
 		SetSelectedBodiesPrimitivesHelper(BodyIndex, AggGeom.TaperedCapsuleElems, NewSelection, Predicate);
+		SetSelectedBodiesPrimitivesHelper(BodyIndex, AggGeom.LevelSetElems, NewSelection, Predicate);
+		SetSelectedBodiesPrimitivesHelper(BodyIndex, AggGeom.SkinnedLevelSetElems, NewSelection, Predicate);
+		SetSelectedBodiesPrimitivesHelper(BodyIndex, AggGeom.MLLevelSetElems, NewSelection, Predicate);
+		SetSelectedBodiesPrimitivesHelper(BodyIndex, AggGeom.SkinnedTriangleMeshElems, NewSelection, Predicate);
 	}
 
 	return NewSelection;
@@ -127,6 +135,105 @@ bool SelectionContainsIndex(CollectionType&& Collection, const int32 InIndex)
 	) != nullptr;
 }
 
+Chaos::TSphere<Chaos::FReal, 3> ConvertPrimitiveToImplicitObject(const FKSphereElem& Elem)
+{
+	return Chaos::TSphere<Chaos::FReal, 3>(FVector::Zero(), Elem.Radius);
+}
+
+Chaos::FCapsule ConvertPrimitiveToImplicitObject(const FKSphylElem& Elem)
+{
+	// FKSphylElem : Axis of Capsule is along the z-axis of the transform.
+	// FCapsule : Requires two end points for construction.
+	
+	const FVector HalfAxis = FVector::ZAxisVector * Elem.Length * 0.5f;
+	return Chaos::FCapsule(-HalfAxis, HalfAxis, Elem.Radius);
+}
+
+Chaos::TBox<Chaos::FReal, 3> ConvertPrimitiveToImplicitObject(const FKBoxElem& Elem)
+{
+	const FVector HalfExtents = FVector(Elem.X, Elem.Y, Elem.Z) * 0.5f;
+	return Chaos::TBox<Chaos::FReal, 3>(-HalfExtents, HalfExtents);
+}
+
+// Returns true if the two supplied primitive shapes overlap.
+template< typename PrimitiveAType, typename PrimitiveBType > bool DoPrimitivesOverlap(const PrimitiveAType& PrimitiveA, const Chaos::FRigidTransform3& BoneTMA, const PrimitiveBType& PrimitiveB, const Chaos::FRigidTransform3& BoneTMB)
+{
+	auto ImplicitObjectA = ConvertPrimitiveToImplicitObject(PrimitiveA);
+	auto ImplicitObjectB = ConvertPrimitiveToImplicitObject(PrimitiveB);
+
+	const FTransform PrimitiveTMA = PrimitiveA.GetTransform() * BoneTMA;
+	const FTransform PrimitiveTMB = PrimitiveB.GetTransform() * BoneTMB;
+
+	return Chaos::Utilities::CastHelper(ImplicitObjectA, PrimitiveTMA, [ImplicitObjectB, PrimitiveTMB](const auto& Downcast, const auto& FullGeomTransform) 
+		{
+			return Chaos::OverlapQuery(ImplicitObjectB, PrimitiveTMB, Downcast, FullGeomTransform, /*Thickness=*/0); 
+		});
+
+	// TODO: Add support for FKTaperedCapsuleElem - currently unsupported by Chaos::OverlapQuery / CastHelper
+}
+
+// Applies an operator to all primitives in the supplied geometry that could be included in an RBAN simulation.
+template< typename OperationType > void ForEachRBANPrimitive(FKAggregateGeom& AggregateGeometry, OperationType& Operation)
+{
+	for (FKSphereElem& Elem : AggregateGeometry.SphereElems)
+	{
+		Operation(Elem);
+	}
+
+	for (FKBoxElem& Elem : AggregateGeometry.BoxElems)
+	{
+		Operation(Elem);
+	}
+
+	for (FKSphylElem& Elem : AggregateGeometry.SphylElems)
+	{
+		Operation(Elem);
+	}
+
+	// TODO: Add support for FKTaperedCapsuleElem
+}
+
+// Returns true if any of the primitives in either supplied body overlap.
+bool DoBodiesOverlap(TObjectPtr<USkeletalBodySetup> BodyA, TObjectPtr<USkeletalBodySetup> BodyB, TObjectPtr<UPhysicsAsset> PhysicsAsset, TObjectPtr<UPhysicsAssetEditorSkeletalMeshComponent> EditorSkelComp)
+{
+	bool bIsOverlapping = false;
+
+	if (EditorSkelComp)
+	{
+		if (const USkeletalMesh* const EditorSkelMesh = PhysicsAsset->GetPreviewMesh())
+		{
+			// Test each geometry object in Body A against each geometry object in Body B - return true if any overlap.
+
+			const FName BoneNameA = BodyA->BoneName;
+			const int32 BoneIndexA = EditorSkelMesh->GetRefSkeleton().FindBoneIndex(BoneNameA);
+			const Chaos::FRigidTransform3 BoneTMA(EditorSkelComp->GetBoneTransform(BoneIndexA).ToMatrixWithScale());
+
+			const FName BoneNameB = BodyB->BoneName;
+			const int32 BoneIndexB = EditorSkelMesh->GetRefSkeleton().FindBoneIndex(BoneNameB);
+			const Chaos::FRigidTransform3 BoneTMB(EditorSkelComp->GetBoneTransform(BoneIndexB).ToMatrixWithScale());
+
+			auto PrimativeOverlapOuter = [&BoneTMA, &BodyB, &BoneTMB, &bIsOverlapping](const auto PrimitiveA)
+				{
+					auto PrimativeOverlapInner = [&PrimitiveA, &BoneTMA, &BoneTMB, &bIsOverlapping](const auto PrimitiveB)
+						{
+							bIsOverlapping |= DoPrimitivesOverlap(PrimitiveA, BoneTMA, PrimitiveB, BoneTMB);
+						};
+
+					ForEachRBANPrimitive(BodyB->AggGeom, PrimativeOverlapInner); // For each geometry object in Body B
+				};
+
+			ForEachRBANPrimitive(BodyA->AggGeom, PrimativeOverlapOuter); // For each geometry object in Body A
+		}
+	}
+
+	return bIsOverlapping;
+}
+
+bool IsBodyPairCollisionEnabled(TObjectPtr<UPhysicsAsset> PhysicsAsset, const int32 BodyAIndex, const int32 BodyBIndex)
+{
+	return !PhysicsAsset->CollisionDisableTable.Find(FRigidBodyIndexPair(BodyAIndex, BodyBIndex));
+}
+
 // class FScopedBulkSelection //
 
 FScopedBulkSelection::FScopedBulkSelection(TSharedPtr<FPhysicsAssetEditorSharedData> InSharedData)
@@ -1019,6 +1126,7 @@ void FPhysicsAssetEditorSharedData::ClearSelected()
 	SelectedObjects->ClearSelection();
 
 	BroadcastSelectionChanged();
+	UpdateNoCollisionBodies();
 }
 
 void FPhysicsAssetEditorSharedData::ClearSelectedPrimitives()
@@ -1523,6 +1631,13 @@ void FPhysicsAssetEditorSharedData::UpdateNoCollisionBodies()
 {
 	NoCollisionBodies.Empty();
 
+	int32 SelectedBodyIndex = INDEX_NONE;
+
+	if (const FPhysicsAssetEditorSharedData::FSelection* const SelectedBodyOrPrimitive = GetSelectedBodyOrPrimitive())
+	{
+		SelectedBodyIndex = SelectedBodyOrPrimitive->Index;
+	}
+
 	// Query disable table with selected body and every other body.
 	for (int32 i = 0; i <PhysicsAsset->SkeletalBodySetups.Num(); ++i)
 	{
@@ -1530,21 +1645,23 @@ void FPhysicsAssetEditorSharedData::UpdateNoCollisionBodies()
 		{
 			continue;
 		}
-		// Add any bodies with bNoCollision
-		if (PhysicsAsset->SkeletalBodySetups[i]->DefaultInstance.GetCollisionEnabled() == ECollisionEnabled::NoCollision)
+
+		if ((SelectedBodyIndex == INDEX_NONE) || (PhysicsAsset->SkeletalBodySetups[i]->DefaultInstance.GetCollisionEnabled() == ECollisionEnabled::NoCollision))
 		{
+			// Add all bodies if non are selected.
+			// Add any bodies with bNoCollision
 			NoCollisionBodies.Add(i);
 		}
-		else if (GetSelectedBodyOrPrimitive() && i != GetSelectedBodyOrPrimitive()->Index)
+		else if (i != SelectedBodyIndex)
 		{
-			if (!ensure(PhysicsAsset->SkeletalBodySetups[GetSelectedBodyOrPrimitive()->Index]))
+			if (!ensure(PhysicsAsset->SkeletalBodySetups[SelectedBodyIndex]))
 			{
 				continue;
 			}
 			// Add this body if it has disabled collision with selected.
-			FRigidBodyIndexPair Key(i, GetSelectedBodyOrPrimitive()->Index);
+			FRigidBodyIndexPair Key(i, SelectedBodyIndex);
 
-			if (PhysicsAsset->SkeletalBodySetups[GetSelectedBodyOrPrimitive()->Index]->DefaultInstance.GetCollisionEnabled() == ECollisionEnabled::NoCollision ||
+			if (PhysicsAsset->SkeletalBodySetups[SelectedBodyIndex]->DefaultInstance.GetCollisionEnabled() == ECollisionEnabled::NoCollision ||
 				PhysicsAsset->CollisionDisableTable.Find(Key))
 			{
 				NoCollisionBodies.Add(i);
@@ -1608,8 +1725,9 @@ void FPhysicsAssetEditorSharedData::SetCollisionBetweenSelected(const bool bEnab
 			}
 		});
 
-	UpdateNoCollisionBodies();	
+	UpdateNoCollisionBodies();
 	RefreshPhysicsAssetChange(PhysicsAsset);
+	InitializeOverlappingBodyPairs();
 	BroadcastPreviewChanged();
 }
 
@@ -1635,14 +1753,16 @@ bool FPhysicsAssetEditorSharedData::CanSetCollisionBetweenSelected(bool bEnableC
 
 void FPhysicsAssetEditorSharedData::SetCollisionBetweenSelectedAndAll(bool bEnableCollision)
 {
-	if (bRunningSimulation || UniqueSelectionReferencingBodies().IsEmpty())
+	FPhysicsAssetEditorSharedData::SelectionUniqueRange SelectedRange = UniqueSelectionReferencingBodies();
+
+	if (bRunningSimulation || SelectedRange.IsEmpty())
 	{
 		return;
 	}
 
 	PhysicsAsset->Modify();
 
-	for(const FSelection& Selection : UniqueSelectionReferencingBodies())
+	for(const FSelection& Selection : SelectedRange)
 	{
 		for(int32 j = 0; j < PhysicsAsset->SkeletalBodySetups.Num(); ++j)
 		{
@@ -1660,6 +1780,7 @@ void FPhysicsAssetEditorSharedData::SetCollisionBetweenSelectedAndAll(bool bEnab
 
 	UpdateNoCollisionBodies();
 	RefreshPhysicsAssetChange(PhysicsAsset);
+	InitializeOverlappingBodyPairs();
 	BroadcastPreviewChanged();
 }
 
@@ -1703,8 +1824,9 @@ void FPhysicsAssetEditorSharedData::SetCollisionBetween(int32 Body1Index, int32
 		}
 
 		UpdateNoCollisionBodies();
-	
 		RefreshPhysicsAssetChange(PhysicsAsset);
+		UpdateOverlappingBodyPairs(Body1Index);
+		UpdateOverlappingBodyPairs(Body2Index);
 	}
 
 	BroadcastPreviewChanged();
@@ -3643,4 +3765,79 @@ void FPhysicsAssetEditorSharedData::EndManipulation()
 	RefreshPhysicsAssetChange(PhysicsAsset, false);
 }
 
+void FPhysicsAssetEditorSharedData::FindOverlappingBodyPairs(const int32 InBodyIndex, TArray<TPair<int32, int32>>& OutCollidingBodyPairs)
+{
+	if (PhysicsAsset->SkeletalBodySetups[InBodyIndex]->DefaultInstance.GetCollisionEnabled() != ECollisionEnabled::NoCollision)
+	{
+		auto CreateCollisionPair = [](const int32 IndexA, const int32 IndexB)
+			{
+				return (IndexA < IndexB) ? TPair<int32, int32>(IndexA, IndexB) : TPair<int32, int32>(IndexB, IndexA);
+			};
+
+		for (int32 BodyIndex = 0, BodyCount = PhysicsAsset->SkeletalBodySetups.Num(); BodyIndex < BodyCount; ++BodyIndex)
+		{
+			if (BodyIndex != InBodyIndex)
+			{
+				if (IsBodyPairCollisionEnabled(PhysicsAsset, InBodyIndex, BodyIndex) && DoBodiesOverlap(PhysicsAsset->SkeletalBodySetups[InBodyIndex], PhysicsAsset->SkeletalBodySetups[BodyIndex], PhysicsAsset, EditorSkelComp))
+				{
+					OutCollidingBodyPairs.AddUnique(CreateCollisionPair(BodyIndex, InBodyIndex));
+				}
+			}
+		}
+	}
+}
+
+void FPhysicsAssetEditorSharedData::RemoveOverlappingBodyPairs(const int32 InBodyIndex, TArray<TPair<int32, int32>>& OutCollidingBodyPairs)
+{
+	auto PairContainsIndex = [InBodyIndex](const TPair<int32, int32>& Element)
+		{
+			return (Element.Key == InBodyIndex) || (Element.Value == InBodyIndex);
+		};
+
+	OutCollidingBodyPairs.RemoveAll(PairContainsIndex);
+}
+
+void FPhysicsAssetEditorSharedData::InitializeOverlappingBodyPairs()
+{
+	OverlappingCollidingBodyPairs.Reset();
+	
+	for (int32 BodyIndex = 0, BodyCount = PhysicsAsset->SkeletalBodySetups.Num(); BodyIndex < BodyCount; ++BodyIndex)
+	{
+		FindOverlappingBodyPairs(BodyIndex, OverlappingCollidingBodyPairs);
+	}
+}
+
+void FPhysicsAssetEditorSharedData::UpdateOverlappingBodyPairs(const int32 InBodyIndex)
+{
+	RemoveOverlappingBodyPairs(InBodyIndex, OverlappingCollidingBodyPairs);
+	FindOverlappingBodyPairs(InBodyIndex, OverlappingCollidingBodyPairs);
+}
+
+bool FPhysicsAssetEditorSharedData::IsBodyOverlapping(const int32 InBodyIndex) const
+{
+	auto PairContainsIndex = [InBodyIndex](const TPair<int32, int32>& Element) // TODO - avoid duplication
+		{
+			return (Element.Key == InBodyIndex) || (Element.Value == InBodyIndex);
+		};
+
+	return Algo::FindByPredicate(OverlappingCollidingBodyPairs, PairContainsIndex) != nullptr;
+}
+
+bool FPhysicsAssetEditorSharedData::ShouldShowBodyOverlappingHighlight(const int32 InBodyIndex) const
+{
+	return IsHighlightingOverlapingBodies() && IsBodyOverlapping(InBodyIndex);
+}
+
+void FPhysicsAssetEditorSharedData::ToggleHighlightOverlapingBodies()
+{
+	EditorOptions->bHighlightOverlapingBodies = ~EditorOptions->bHighlightOverlapingBodies;
+	EditorOptions->SaveConfig();
+}
+
+bool FPhysicsAssetEditorSharedData::IsHighlightingOverlapingBodies() const
+{
+	return EditorOptions->bHighlightOverlapingBodies;
+}
+
+
 #undef LOCTEXT_NAMESPACE
diff --git a/Engine/Source/Editor/PhysicsAssetEditor/Private/PhysicsAssetEditorSharedData.h b/Engine/Source/Editor/PhysicsAssetEditor/Private/PhysicsAssetEditorSharedData.h
--- a/Engine/Source/Editor/PhysicsAssetEditor/Private/PhysicsAssetEditorSharedData.h
+++ b/Engine/Source/Editor/PhysicsAssetEditor/Private/PhysicsAssetEditorSharedData.h
@@ -460,6 +460,19 @@ public:
 
 	struct FPhysicsAssetRenderSettings* GetRenderSettings() const;
 
+	/** Physics Body Overlap Detection **/
+	void FindOverlappingBodyPairs(const int32 InBodyIndex, TArray<TPair<int32, int32>>& OutCollidingBodyPairs);
+	void RemoveOverlappingBodyPairs(const int32 InBodyIndex, TArray<TPair<int32, int32>>& OutCollidingBodyPairs);
+	void InitializeOverlappingBodyPairs();
+	void UpdateOverlappingBodyPairs(const int32 InBodyIndex);
+	bool IsBodyOverlapping(const int32 InBodyIndex) const;
+	bool ShouldShowBodyOverlappingHighlight(const int32 InBodyIndex) const;
+
+	void ToggleHighlightOverlapingBodies();
+	bool IsHighlightingOverlapingBodies() const;
+
+	TArray<TPair<int32, int32>> OverlappingCollidingBodyPairs; // A record of all the pairs of physics bodies that are overlapping and not flagged as non-colliding in the physics asset.
+
 	/** Misc toggles */
 	bool bRunningSimulation;
 	bool bNoGravitySimulation;
diff --git a/Engine/Source/Editor/PhysicsAssetEditor/Private/PhysicsAssetEditorSkeletalMeshComponent.cpp b/Engine/Source/Editor/PhysicsAssetEditor/Private/PhysicsAssetEditorSkeletalMeshComponent.cpp
--- a/Engine/Source/Editor/PhysicsAssetEditor/Private/PhysicsAssetEditorSkeletalMeshComponent.cpp
+++ b/Engine/Source/Editor/PhysicsAssetEditor/Private/PhysicsAssetEditorSkeletalMeshComponent.cpp
@@ -33,11 +33,26 @@ namespace
 	FAutoConsoleVariableRef CVarChaosImmPhysStepTime(TEXT("p.PhAT.DebugViewportClicks"), bDebugViewportClicks, TEXT("Set to 1 to show mouse click results in PhAT"));
 }
 
+// struct FPhysicsAssetEditorDrawState //
+
+FPhysicsAssetEditorDrawState::FPhysicsAssetEditorDrawState(TObjectPtr<UMaterialInstanceDynamic> InMaterial, const FColor& InColor)
+: Material(InMaterial)
+, Color(InColor)
+{
+	check(Material);
+}
+
+FPhysicsAssetEditorDrawState::FPhysicsAssetEditorDrawState(const TCHAR* const MaterialName, const FColor& InColor)
+: Color(InColor)
+{
+	UMaterialInterface* const BaseMaterial = LoadObject<UMaterialInterface>(NULL, MaterialName, NULL, LOAD_None, NULL);
+	Material = UMaterialInstanceDynamic::Create(BaseMaterial, GetTransientPackage());
+	check(Material);
+}
+
+
 UPhysicsAssetEditorSkeletalMeshComponent::UPhysicsAssetEditorSkeletalMeshComponent(const FObjectInitializer& ObjectInitializer)
 	: Super(ObjectInitializer)
-	, BoneUnselectedColor(170, 155, 225)
-	, NoCollisionColor(200, 200, 200)
-	, FixedColor(125, 125, 0)
 	, ConstraintBone1Color(255, 166, 0)
 	, ConstraintBone2Color(0, 150, 150)
 	, HierarchyDrawColor(220, 255, 220)
@@ -48,22 +63,23 @@ UPhysicsAssetEditorSkeletalMeshComponent::UPhysicsAssetEditorSkeletalMeshCompone
 {
 	if (!HasAnyFlags(RF_DefaultSubObject | RF_ArchetypeObject | RF_ClassDefaultObject))
 	{
+		static FName SelectionColorName(TEXT("SelectionColor"));
+		const FSlateColor SelectionColor = FAppStyle::GetSlateColor(SelectionColorName);
+		const FLinearColor SelectionColorLinear(SelectionColor.IsColorSpecified() ? SelectionColor.GetSpecifiedColor() : FLinearColor::White);
+
 		// Body materials
-		UMaterialInterface* BaseElemSelectedMaterial = LoadObject<UMaterialInterface>(NULL, TEXT("/Engine/EditorMaterials/PhAT_ElemSelectedMaterial.PhAT_ElemSelectedMaterial"), NULL, LOAD_None, NULL);
-		ElemSelectedMaterial = UMaterialInstanceDynamic::Create(BaseElemSelectedMaterial, GetTransientPackage());
-		check(ElemSelectedMaterial);
+		ElemSelectedPrimitiveDrawState = FPhysicsAssetEditorDrawState(TEXT("/Engine/EditorMaterials/PhAT_ElemSelectedMaterial.PhAT_ElemSelectedMaterial"), SelectionColorLinear.ToFColor(true));
+		ElemPrimitiveInSelectedBodyDrawState = FPhysicsAssetEditorDrawState(TEXT("/Engine/EditorMaterials/PhAT_ElemPrimitiveInSelectedBodyMaterial.PhAT_ElemPrimitiveInSelectedBodyMaterial"), (SelectionColorLinear * 0.5f).ToFColor(true));
+		ElemUnselectedDrawState = FPhysicsAssetEditorDrawState(TEXT("/Engine/EditorMaterials/PhAT_ElemUnselectedMaterial.PhAT_ElemUnselectedMaterial"), FColor(97, 102, 102));
+		ElemSelectedOverlappingDrawState = FPhysicsAssetEditorDrawState(TEXT("/Engine/EditorMaterials/PhAT_ElemSelectedOverlappingMaterial.PhAT_ElemSelectedOverlappingMaterial"), FColor(102, 20, 20));
+		ElemUnselectedOverlappingDrawState = FPhysicsAssetEditorDrawState(TEXT("/Engine/EditorMaterials/PhAT_ElemUnselectedOverlappingMaterial.PhAT_ElemUnselectedOverlappingMaterial"), FColor(102, 20, 20));
+		ElemCollidingWithSelectedDrawState = FPhysicsAssetEditorDrawState(TEXT("/Engine/EditorMaterials/PhAT_ElemCollidingWithSelectedMaterial.PhAT_ElemCollidingWithSelectedMaterial"), FColor(255, 140, 0));
+		BoneUnselectedDrawState = FPhysicsAssetEditorDrawState(TEXT("/Engine/EditorMaterials/PhAT_UnselectedMaterial.PhAT_UnselectedMaterial"), FColor(97, 102, 102));
+		BoneNoCollisionDrawState = FPhysicsAssetEditorDrawState(TEXT("/Engine/EditorMaterials/PhAT_NoCollisionMaterial.PhAT_NoCollisionMaterial"), FColor(128, 128, 128));
 
 		BoneMaterialHit = UMaterial::GetDefaultMaterial(MD_Surface);
 		check(BoneMaterialHit);
 
-		UMaterialInterface* BaseBoneUnselectedMaterial = LoadObject<UMaterialInterface>(NULL, TEXT("/Engine/EditorMaterials/PhAT_UnselectedMaterial.PhAT_UnselectedMaterial"), NULL, LOAD_None, NULL);
-		BoneUnselectedMaterial = UMaterialInstanceDynamic::Create(BaseBoneUnselectedMaterial, GetTransientPackage());
-		check(BoneUnselectedMaterial);
-
-		UMaterialInterface* BaseBoneNoCollisionMaterial = LoadObject<UMaterialInterface>(NULL, TEXT("/Engine/EditorMaterials/PhAT_NoCollisionMaterial.PhAT_NoCollisionMaterial"), NULL, LOAD_None, NULL);
-		BoneNoCollisionMaterial = UMaterialInstanceDynamic::Create(BaseBoneNoCollisionMaterial, GetTransientPackage());
-		check(BoneNoCollisionMaterial);
-
 		// this is because in phat editor, you'd like to see fixed bones to be fixed without animation force update
 		KinematicBonesUpdateType = EKinematicBonesUpdateToPhysics::SkipSimulatingBones;
 		bUpdateJointsFromAnimation = false;
@@ -105,15 +121,16 @@ void UPhysicsAssetEditorSkeletalMeshComponent::DebugDraw(const FSceneView* View,
 
 	// set opacity of our materials
 	static FName OpacityName(TEXT("Opacity"));
-	ElemSelectedMaterial->SetScalarParameterValue(OpacityName, SharedData->EditorOptions->CollisionOpacity);
-	BoneUnselectedMaterial->SetScalarParameterValue(OpacityName, SharedData->EditorOptions->bSolidRenderingForSelectedOnly ? 0.0f : SharedData->EditorOptions->CollisionOpacity);
-	BoneNoCollisionMaterial->SetScalarParameterValue(OpacityName, SharedData->EditorOptions->bSolidRenderingForSelectedOnly ? 0.0f : SharedData->EditorOptions->CollisionOpacity);
+	ElemSelectedPrimitiveDrawState.Material->SetScalarParameterValue(OpacityName, SharedData->EditorOptions->CollisionOpacity);
+	ElemPrimitiveInSelectedBodyDrawState.Material->SetScalarParameterValue(OpacityName, SharedData->EditorOptions->CollisionOpacity);
+	BoneUnselectedDrawState.Material->SetScalarParameterValue(OpacityName, SharedData->EditorOptions->bSolidRenderingForSelectedOnly ? 0.0f : SharedData->EditorOptions->CollisionOpacity);
+	BoneNoCollisionDrawState.Material->SetScalarParameterValue(OpacityName, SharedData->EditorOptions->bSolidRenderingForSelectedOnly ? 0.0f : SharedData->EditorOptions->CollisionOpacity);
 
 	static FName SelectionColorName(TEXT("SelectionColor"));
 	const FSlateColor SelectionColor = FAppStyle::GetSlateColor(SelectionColorName);
 	const FLinearColor LinearSelectionColor(SelectionColor.IsColorSpecified() ? SelectionColor.GetSpecifiedColor() : FLinearColor::White);
 
-	ElemSelectedMaterial->SetVectorParameterValue(SelectionColorName, LinearSelectionColor);
+	ElemSelectedPrimitiveDrawState.Material->SetVectorParameterValue(SelectionColorName, LinearSelectionColor);
 
 	FPhysicsAssetRenderSettings* const RenderSettings = UPhysicsAssetRenderUtilities::GetSettings(PhysicsAsset);
 	
@@ -244,105 +261,89 @@ FTransform UPhysicsAssetEditorSkeletalMeshComponent::GetPrimitiveTransform(const
 	return FTransform::Identity;
 }
 
-FColor UPhysicsAssetEditorSkeletalMeshComponent::GetPrimitiveColor(const int32 BodyIndex, const EAggCollisionShape::Type PrimitiveType, const int32 PrimitiveIndex) const
+const FPhysicsAssetEditorDrawState& UPhysicsAssetEditorSkeletalMeshComponent::GetPrimitiveDrawState(const int32 BodyIndex, const EAggCollisionShape::Type PrimitiveType, const int32 PrimitiveIndex) const
 {
-	UBodySetup* SharedBodySetup = SharedData->PhysicsAsset->SkeletalBodySetups[BodyIndex];
-
-	if (!SharedData->bRunningSimulation && SharedData->GetSelectedConstraint())
+	if (SharedData->bRunningSimulation)
 	{
-		UPhysicsConstraintTemplate* cs = SharedData->PhysicsAsset->ConstraintSetup[SharedData->GetSelectedConstraint()->Index];
-
-		if (cs->DefaultInstance.ConstraintBone1 == SharedBodySetup->BoneName)
-		{
-			return ConstraintBone1Color;
-		}
-		else if (cs->DefaultInstance.ConstraintBone2 == SharedBodySetup->BoneName)
-		{
-			return ConstraintBone2Color;
-		}
+		return PrimitiveType == EAggCollisionShape::TaperedCapsule ? BoneNoCollisionDrawState : BoneUnselectedDrawState;
 	}
 
-	static FName SelectionColorName(TEXT("SelectionColor"));
-	const FSlateColor SelectionColor = FAppStyle::GetSlateColor(SelectionColorName);
-	const FLinearColor SelectionColorLinear(SelectionColor.IsColorSpecified() ? SelectionColor.GetSpecifiedColor() : FLinearColor::White);
-	const FColor ElemSelectedColor = SelectionColorLinear.ToFColor(true);
-	const FColor ElemSelectedBodyColor = (SelectionColorLinear* 0.5f).ToFColor(true);
+	const bool bIsOverlapping = SharedData->ShouldShowBodyOverlappingHighlight(BodyIndex);
+	bool bIsSelectedPrimitive = false;
+	bool bIsSelectedBody = false;
 
-	bool bInBody = false;
 	for (const FPhysicsAssetEditorSharedData::FSelection& SelectedElement : SharedData->SelectedPrimitives())
 	{
-		if (BodyIndex == SelectedElement.Index)
+		if (IsReferencingPrimitive(SelectedElement, BodyIndex, PrimitiveType, PrimitiveIndex))
 		{
-			bInBody = true;
+			bIsSelectedPrimitive = true;
+			bIsSelectedBody = true;
+			break;
 		}
-
-		if (IsReferencingPrimitive(SelectedElement, BodyIndex, PrimitiveType, PrimitiveIndex) && !SharedData->bRunningSimulation)
+		else if (SelectedElement.GetIndex() == BodyIndex)
 		{
-			return ElemSelectedColor;
+			bIsSelectedBody = true;
 		}
 	}
 
-	if (bInBody && !SharedData->bRunningSimulation)	//this primitive is in a body that's currently selected, but this primitive itself isn't selected
+	if (bIsSelectedPrimitive && bIsOverlapping)
 	{
-		return ElemSelectedBodyColor;
+		// This selected primitive should be highlighted as being part of a body that is overlapping the selected body.
+		return ElemSelectedOverlappingDrawState;
 	}
-	if(PrimitiveType == EAggCollisionShape::TaperedCapsule)
+	else if (bIsOverlapping)
 	{
-		return NoCollisionColor;
+		// This primitive should be highlighted as being part of a body that is overlapping the selected body.
+		return ElemUnselectedOverlappingDrawState;
 	}
-
-	if (SharedData->bRunningSimulation)
+	else if (bIsSelectedPrimitive)
 	{
-		const bool bIsSimulatedAtAll = SharedBodySetup->PhysicsType == PhysType_Simulated || (SharedBodySetup->PhysicsType == PhysType_Default && SharedData->EditorOptions->PhysicsBlend > 0.f);
-		if (!bIsSimulatedAtAll)
-		{
-			return FixedColor;
-		}
+		// This primitive is the selected element.
+		return ElemSelectedPrimitiveDrawState;
+	}
+	else if (bIsSelectedBody)
+	{
+		// This primitive is a child of a selected body.
+		return ElemPrimitiveInSelectedBodyDrawState;
+	}
+	else if ((PrimitiveType == EAggCollisionShape::TaperedCapsule) || (SharedData->NoCollisionBodies.Find(BodyIndex) != INDEX_NONE && !SharedData->bRunningSimulation))
+	{
+		// If there is no collision with this body, use 'no collision material'.
+		return BoneNoCollisionDrawState;
 	}
 	else
 	{
-		if (!SharedData->bRunningSimulation && !SharedData->SelectedBodiesAndPrimitives().IsEmpty())
-		{
-			// If there is no collision with this body, use 'no collision material'.
-			if (SharedData->NoCollisionBodies.Find(BodyIndex) != INDEX_NONE)
-			{
-				return NoCollisionColor;
-			}
-		}
+		// Collisions are enabled between this body and the selected body.
+		return ElemCollidingWithSelectedDrawState;
 	}
-
-	return BoneUnselectedColor;
 }
 
-UMaterialInterface* UPhysicsAssetEditorSkeletalMeshComponent::GetPrimitiveMaterial(const int32 BodyIndex, const EAggCollisionShape::Type PrimitiveType, const int32 PrimitiveIndex) const
+
+
+FColor UPhysicsAssetEditorSkeletalMeshComponent::GetPrimitiveColor(const int32 BodyIndex, const EAggCollisionShape::Type PrimitiveType, const int32 PrimitiveIndex) const
 {
-	if (SharedData->bRunningSimulation)
-	{
-		return PrimitiveType == EAggCollisionShape::TaperedCapsule ? BoneNoCollisionMaterial : BoneUnselectedMaterial;
-	}
+	UBodySetup* SharedBodySetup = SharedData->PhysicsAsset->SkeletalBodySetups[BodyIndex];
 
-	for (const FPhysicsAssetEditorSharedData::FSelection& SelectedElement : SharedData->SelectedBodiesAndPrimitives())
+	if (!SharedData->bRunningSimulation && SharedData->GetSelectedConstraint())
 	{
-		if (IsReferencingPrimitive(SelectedElement, BodyIndex, PrimitiveType, PrimitiveIndex))
+		UPhysicsConstraintTemplate* cs = SharedData->PhysicsAsset->ConstraintSetup[SharedData->GetSelectedConstraint()->Index];
+
+		if (cs->DefaultInstance.ConstraintBone1 == SharedBodySetup->BoneName)
 		{
-			return ElemSelectedMaterial;
+			return ConstraintBone1Color;
+		}
+		else if (cs->DefaultInstance.ConstraintBone2 == SharedBodySetup->BoneName)
+		{
+			return ConstraintBone2Color;
 		}
 	}
 
-	if (PrimitiveType == EAggCollisionShape::TaperedCapsule)
-	{
-		return BoneNoCollisionMaterial;
-	}
+	return GetPrimitiveDrawState(BodyIndex, PrimitiveType, PrimitiveIndex).Color;
+}
 
-	// If there is no collision with this body, use 'no collision material'.
-	if (!SharedData->SelectedBodiesAndPrimitives().IsEmpty() && SharedData->NoCollisionBodies.Find(BodyIndex) != INDEX_NONE && !SharedData->bRunningSimulation)
-	{
-		return BoneNoCollisionMaterial;
-	}
-	else
-	{
-		return BoneUnselectedMaterial;
-	}
+UMaterialInterface* UPhysicsAssetEditorSkeletalMeshComponent::GetPrimitiveMaterial(const int32 BodyIndex, const EAggCollisionShape::Type PrimitiveType, const int32 PrimitiveIndex) const
+{
+	return GetPrimitiveDrawState(BodyIndex, PrimitiveType, PrimitiveIndex).Material;
 }
 
 void UPhysicsAssetEditorSkeletalMeshComponent::RefreshBoneTransforms(FActorComponentTickFunction* TickFunction)
diff --git a/Engine/Source/Editor/PhysicsAssetEditor/Private/PhysicsAssetEditorSkeletalMeshComponent.h b/Engine/Source/Editor/PhysicsAssetEditor/Private/PhysicsAssetEditorSkeletalMeshComponent.h
--- a/Engine/Source/Editor/PhysicsAssetEditor/Private/PhysicsAssetEditorSkeletalMeshComponent.h
+++ b/Engine/Source/Editor/PhysicsAssetEditor/Private/PhysicsAssetEditorSkeletalMeshComponent.h
@@ -12,6 +12,28 @@ class FPrimitiveDrawInterface;
 class UMaterialInterface;
 class UMaterialInstanceDynamic;
 
+
+/**
+ * FDrawState
+ *
+ * Keeps track of render material and line color for an object or state displayed in the view port. 
+ */
+USTRUCT() 
+struct FPhysicsAssetEditorDrawState
+{
+	GENERATED_BODY()
+	FPhysicsAssetEditorDrawState() = default;
+	FPhysicsAssetEditorDrawState(TObjectPtr<UMaterialInstanceDynamic> InMaterial, const FColor& InColor);
+	FPhysicsAssetEditorDrawState(const TCHAR* MaterialName, const FColor& InColor);
+	
+	UPROPERTY(transient)
+	TObjectPtr<UMaterialInstanceDynamic> Material;
+
+	UPROPERTY(transient)
+	FColor Color;
+};
+
+
 UCLASS()
 class UPhysicsAssetEditorSkeletalMeshComponent : public UDebugSkelMeshComponent
 {
@@ -20,11 +42,38 @@ class UPhysicsAssetEditorSkeletalMeshComponent : public UDebugSkelMeshComponent
 	/** Data and methods shared across multiple classes */
 	class FPhysicsAssetEditorSharedData* SharedData;
 
-	// Draw colors
+	// Draw States
+
+	// Primitives that are directly selected.
+	UPROPERTY(transient)
+	FPhysicsAssetEditorDrawState ElemSelectedPrimitiveDrawState;
+	
+	// Primitives that are part of a selected body but not directly selected.
+	UPROPERTY(transient)
+	FPhysicsAssetEditorDrawState ElemPrimitiveInSelectedBodyDrawState;
+
+	// Bodies that are currently not selected.
+	UPROPERTY(transient)
+	FPhysicsAssetEditorDrawState ElemUnselectedDrawState;
+
+	// Bodies that are currently selected and would collide with other bodies in the current pose during simulation.
+	UPROPERTY(transient)
+	FPhysicsAssetEditorDrawState ElemSelectedOverlappingDrawState;
+
+	// Bodies that are currently not selected and would collide with other bodies in the current pose during simulation.
+	UPROPERTY(transient)
+	FPhysicsAssetEditorDrawState ElemUnselectedOverlappingDrawState;
+
+	// Bodies that are able to collide with one or more of the selected bodies.
+	UPROPERTY(transient)
+	FPhysicsAssetEditorDrawState ElemCollidingWithSelectedDrawState;
+
+	UPROPERTY(transient)
+	FPhysicsAssetEditorDrawState BoneUnselectedDrawState;
+
+	UPROPERTY(transient)
+	FPhysicsAssetEditorDrawState BoneNoCollisionDrawState;
 
-	FColor BoneUnselectedColor;
-	FColor NoCollisionColor;
-	FColor FixedColor;
 	FColor ConstraintBone1Color;
 	FColor ConstraintBone2Color;
 	FColor HierarchyDrawColor;
@@ -33,15 +82,6 @@ class UPhysicsAssetEditorSkeletalMeshComponent : public UDebugSkelMeshComponent
 	float InfluenceLineLength;
 	FColor InfluenceLineColor;
 
-	// Materials
-
-	UPROPERTY(transient)
-	TObjectPtr<UMaterialInstanceDynamic> ElemSelectedMaterial;
-	UPROPERTY(transient)
-	TObjectPtr<UMaterialInstanceDynamic> BoneUnselectedMaterial;
-	UPROPERTY(transient)
-	TObjectPtr<UMaterialInstanceDynamic> BoneNoCollisionMaterial;
-
 	UPROPERTY(transient)
 	TObjectPtr<UMaterialInterface> BoneMaterialHit;
 
@@ -80,6 +120,8 @@ public:
 	virtual bool CanOverrideCollisionProfile() const override { return false;  }
 
 private:
+	const FPhysicsAssetEditorDrawState& GetPrimitiveDrawState(const int32 BodyIndex, const EAggCollisionShape::Type PrimitiveType, const int32 PrimitiveIndex) const;
+
 	void UpdateSkinnedLevelSets();
 	void UpdateMLLevelSets();
 	void UpdateSkinnedTriangleMeshes();
diff --git a/Engine/Source/Editor/UnrealEd/Classes/Preferences/PhysicsAssetEditorOptions.h b/Engine/Source/Editor/UnrealEd/Classes/Preferences/PhysicsAssetEditorOptions.h
--- a/Engine/Source/Editor/UnrealEd/Classes/Preferences/PhysicsAssetEditorOptions.h
+++ b/Engine/Source/Editor/UnrealEd/Classes/Preferences/PhysicsAssetEditorOptions.h
@@ -111,6 +111,10 @@ class UPhysicsAssetEditorOptions : public UObject
 	UPROPERTY(EditAnywhere, config, Category = Poking, meta = (ClampMin = 0))
 	float InteractionDistance;
 
+	/** Whether to draw overlapping physics bodies in a different color in the view port. */
+	UPROPERTY(config)
+	uint32 bHighlightOverlapingBodies : 1;
+
 	/** Whether to draw constraints as points */
 	UPROPERTY(config)
 	uint32 bShowConstraintsAsPoints:1;
