[UE] Ensure PhAT body overlap detection is updated when bodies are mirrored or pasted from clipboard.

#rnx

diff --git a/Engine/Source/Editor/PhysicsAssetEditor/Private/PhysicsAssetEditorSharedData.cpp b/Engine/Source/Editor/PhysicsAssetEditor/Private/PhysicsAssetEditorSharedData.cpp
--- a/Engine/Source/Editor/PhysicsAssetEditor/Private/PhysicsAssetEditorSharedData.cpp
+++ b/Engine/Source/Editor/PhysicsAssetEditor/Private/PhysicsAssetEditorSharedData.cpp
@@ -756,6 +756,8 @@ void FPhysicsAssetEditorSharedData::Mirror()
 					UPhysicsConstraintTemplate * ToConstraint = PhysicsAsset->ConstraintSetup[MirrorConstraintIndex];
 					CopyConstraintProperties(FromConstraint, ToConstraint);
 				}
+
+				UpdateOverlappingBodyPairs(MirrorBodyIndex);
 			}
 		}
 
@@ -997,6 +999,8 @@ void FPhysicsAssetEditorSharedData::RefreshPhysicsAssetChange(const UPhysicsAsse
 			UpdateClothPhysics();
 		}
 		EnableSimulation(false);
+
+		InitializeOverlappingBodyPairs();
 	}
 }
 
@@ -3668,6 +3672,7 @@ void FPhysicsAssetEditorSharedData::PostUndo()
 	BroadcastPreviewChanged();
 	BroadcastHierarchyChanged();
 	BroadcastSelectionChanged();
+	InitializeOverlappingBodyPairs();
 }
 
 void FPhysicsAssetEditorSharedData::Redo()
@@ -3767,7 +3772,7 @@ void FPhysicsAssetEditorSharedData::EndManipulation()
 
 void FPhysicsAssetEditorSharedData::FindOverlappingBodyPairs(const int32 InBodyIndex, TArray<TPair<int32, int32>>& OutCollidingBodyPairs)
 {
-	if (PhysicsAsset->SkeletalBodySetups[InBodyIndex]->DefaultInstance.GetCollisionEnabled() != ECollisionEnabled::NoCollision)
+	if (PhysicsAsset->SkeletalBodySetups.IsValidIndex(InBodyIndex) && (PhysicsAsset->SkeletalBodySetups[InBodyIndex]->DefaultInstance.GetCollisionEnabled() != ECollisionEnabled::NoCollision))
 	{
 		auto CreateCollisionPair = [](const int32 IndexA, const int32 IndexB)
 			{
