Chaos Cloth Asset - Fixed a simulation crash while attempting to use a physics asset containing convex hulls, right after that its physics asset editor has been closed.

[RN] Chaos Cloth Asset - Fixed a simulation crash while attempting to use a physics asset containing convex hulls, right after that its physics asset editor has been closed.

diff --git a/Engine/Source/Editor/PhysicsAssetEditor/Private/PhysicsAssetEditor.cpp b/Engine/Source/Editor/PhysicsAssetEditor/Private/PhysicsAssetEditor.cpp
--- a/Engine/Source/Editor/PhysicsAssetEditor/Private/PhysicsAssetEditor.cpp
+++ b/Engine/Source/Editor/PhysicsAssetEditor/Private/PhysicsAssetEditor.cpp
@@ -145,6 +145,28 @@ namespace PhysicsAssetEditor
 
 		return TSharedPtr<FPhysicsAssetEditor>();
 	}
+
+	static void RecreatePhysicsAssetConvexMeshes(UPhysicsAsset* PhysicsAsset)
+	{
+		if (PhysicsAsset)
+		{
+			for (TObjectPtr<USkeletalBodySetup>& SkeletalBodySetup : PhysicsAsset->SkeletalBodySetups)
+			{
+				if (SkeletalBodySetup)
+				{
+					for (const FKConvexElem& Element : SkeletalBodySetup->AggGeom.ConvexElems)
+					{
+						if (!Element.GetChaosConvexMesh())
+						{
+							SkeletalBodySetup->InvalidatePhysicsData();
+							SkeletalBodySetup->CreatePhysicsMeshes();
+							break;
+						}
+					}
+				}
+			}
+		}
+	}
 }
 
 void FPhysicsAssetEditor::RegisterTabSpawners(const TSharedRef<class FTabManager>& InTabManager)
@@ -168,6 +190,9 @@ FPhysicsAssetEditor::~FPhysicsAssetEditor()
 	{
 		constexpr bool bSetPreviewMeshInAsset = false;
 		PersonaToolkit->SetPreviewMesh(nullptr, bSetPreviewMeshInAsset);
+
+		// Recreate the physics meshes since SetPreviewMesh invalidates and clears the physics meshes, causing the ChaosConvexMesh data to be lost in the process
+		PhysicsAssetEditor::RecreatePhysicsAssetConvexMeshes(SharedData->PhysicsAsset);
 	}
 }
 
@@ -489,29 +514,7 @@ void FPhysicsAssetEditor::PostRedo( bool bSuccess )
 {
 	OnPostUndo.Broadcast();
 
-	for (int32 BodyIdx=0; BodyIdx < SharedData->PhysicsAsset->SkeletalBodySetups.Num(); ++BodyIdx)
-	{
-		UBodySetup* Body = SharedData->PhysicsAsset->SkeletalBodySetups[BodyIdx];
-		
-		bool bRecreate = false;
-		for (int32 ElemIdx=0; ElemIdx < Body->AggGeom.ConvexElems.Num(); ++ElemIdx)
-		{
-			FKConvexElem& Element = Body->AggGeom.ConvexElems[ElemIdx];
-
-			if (Element.GetChaosConvexMesh() == NULL)
-			{
-				bRecreate = true;
-				break;
-			}
-		}
-
-		if (bRecreate)
-		{
-			Body->InvalidatePhysicsData();
-			Body->CreatePhysicsMeshes();
-		}
-
-	}
+	PhysicsAssetEditor::RecreatePhysicsAssetConvexMeshes(SharedData->PhysicsAsset);
 
 	PostUndo(bSuccess);
 }
