[Editor TRS] Add/RemoveViewportOverlay and other minor additions to support new Gizmos in:
- BP Editor
- Persona Editors: Animation, Skeleton, SkeletalMesh, PHAT

#Synced-CL 51856350

diff --git a/Engine/Source/Editor/AnimationEditor/Private/AnimationEditor.cpp b/Engine/Source/Editor/AnimationEditor/Private/AnimationEditor.cpp
--- a/Engine/Source/Editor/AnimationEditor/Private/AnimationEditor.cpp
+++ b/Engine/Source/Editor/AnimationEditor/Private/AnimationEditor.cpp
@@ -24,6 +24,7 @@
 #include "Curves/RichCurve.h"
 #include "Editor.h"
 #include "Editor/EditorEngine.h"
+#include "EditorModeManager.h"
 #include "EditorReimportHandler.h"
 #include "Engine/CurveTable.h"
 #include "Engine/SkeletalMesh.h"
@@ -50,6 +51,7 @@
 #include "ISkeletonTree.h"
 #include "ISkeletonTreeItem.h"
 #include "Internationalization/Internationalization.h"
+#include "IPersonaViewport.h"
 #include "Logging/LogMacros.h"
 #include "Misc/AssertionMacros.h"
 #include "Misc/Attribute.h"
@@ -73,6 +75,7 @@
 #include "ToolMenuSection.h"
 #include "ToolMenus.h"
 #include "Toolkits/AssetEditorToolkit.h"
+#include "Tools/EdModeInteractiveToolsContext.h"
 #include "UObject/Object.h"
 #include "UObject/ObjectMacros.h"
 #include "UObject/ObjectPtr.h"
@@ -422,6 +425,22 @@ void FAnimationEditor::HandleSelectionChanged(const TArrayView<TSharedPtr<ISkele
 	}
 }
 
+void FAnimationEditor::AddViewportOverlayWidget(TSharedRef<SWidget> InOverlaidWidget, int32 ZOrder)
+{
+	if (Viewport.IsValid())
+	{
+		Viewport->AddOverlayWidget(InOverlaidWidget, ZOrder);
+	}
+}
+
+void FAnimationEditor::RemoveViewportOverlayWidget(TSharedRef<SWidget> InOverlaidWidget)
+{
+	if (Viewport.IsValid())
+	{
+		Viewport->RemoveOverlayWidget(InOverlaidWidget);
+	}
+}
+
 void FAnimationEditor::HandleObjectSelected(UObject* InObject)
 {
 	if (DetailsView.IsValid())
diff --git a/Engine/Source/Editor/AnimationEditor/Private/AnimationEditor.h b/Engine/Source/Editor/AnimationEditor/Private/AnimationEditor.h
--- a/Engine/Source/Editor/AnimationEditor/Private/AnimationEditor.h
+++ b/Engine/Source/Editor/AnimationEditor/Private/AnimationEditor.h
@@ -121,6 +121,10 @@ public:
 
 	void HandleObjectsSelected(const TArray<UObject*>& InObjects);
 
+	// FAssetEditorToolkit interface
+	virtual void AddViewportOverlayWidget(TSharedRef<SWidget> InOverlaidWidget, int32 ZOrder = INDEX_NONE) override;
+	virtual void RemoveViewportOverlayWidget(TSharedRef<SWidget> InOverlaidWidget) override;
+
 private:
 
 	/** Options for asset export */
diff --git a/Engine/Source/Editor/Kismet/Private/BlueprintEditor.cpp b/Engine/Source/Editor/Kismet/Private/BlueprintEditor.cpp
--- a/Engine/Source/Editor/Kismet/Private/BlueprintEditor.cpp
+++ b/Engine/Source/Editor/Kismet/Private/BlueprintEditor.cpp
@@ -187,6 +187,7 @@
 #include "Editor/EditorEngine.h"
 #include "EditorViewportSelectabilityBridge.h"
 #include "IBlueprintSelectionNotifier.h"
+#include "Tools/EdModeInteractiveToolsContext.h"
 
 DEFINE_LOG_CATEGORY_STATIC(LogBlueprintEditor, Log, All);
 
@@ -1334,6 +1335,9 @@ void FBlueprintEditor::OnSelectionUpdated(const TArray<TSharedPtr<FSubobjectEdit
 			}
 		}
 
+		const bool bHasSceneComponent = InspectorObjects.ContainsByPredicate([](const UObject* Obj) { return Obj && Obj->IsA<USceneComponent>(); });
+		GetEditorModeManager().SetSelectionHasSceneComponent(bHasSceneComponent);
+
 		// Update the details panel
 		SKismetInspector::FShowDetailsOptions Options(InspectorTitle, true);
 		Options.bShowComponents = bShowComponents;
@@ -2575,8 +2579,8 @@ void FBlueprintEditor::InitBlueprintEditor(
 		{
 			if (FEditorModeTools* const ModeTools = ViewportClient->GetModeTools())
 			{
-				// Opt out of new ITF gizmos and viewport inputs
-				ModeTools->SetSupportsViewportITF(false);
+				// Opt int to new ITF gizmos and viewport inputs
+				ModeTools->SetSupportsViewportITF(true);
 			}
 		}
 	}
diff --git a/Engine/Source/Editor/PhysicsAssetEditor/Private/PhysicsAssetEditor.cpp b/Engine/Source/Editor/PhysicsAssetEditor/Private/PhysicsAssetEditor.cpp
--- a/Engine/Source/Editor/PhysicsAssetEditor/Private/PhysicsAssetEditor.cpp
+++ b/Engine/Source/Editor/PhysicsAssetEditor/Private/PhysicsAssetEditor.cpp
@@ -57,6 +57,7 @@
 #include "PhysicsEngine/PhysicsAsset.h"
 #include "PhysicsEngine/SkeletalBodySetup.h"
 #include "Engine/Selection.h"
+#include "IPersonaViewport.h"
 #include "PersonaModule.h"
 #include "PersonaToolMenuContext.h"
 
@@ -92,6 +93,7 @@
 #include "AnimationEditorPreviewActor.h"
 #include "Preferences/PersonaOptions.h"
 #include "PhysicsEngine/MLLevelSetModelAndBonesBinningInfo.h"
+#include "Tools/EdModeInteractiveToolsContext.h"
 
 const FName PhysicsAssetEditorModes::PhysicsAssetEditorMode("PhysicsAssetEditorMode");
 
@@ -322,6 +324,27 @@ TSharedPtr<FPhysicsAssetEditorSharedData> FPhysicsAssetEditor::GetSharedData() c
 	return SharedData;
 }
 
+void FPhysicsAssetEditor::HandleViewportCreated(const TSharedRef<IPersonaViewport>& InViewport)
+{
+	Viewport = InViewport;
+}
+
+void FPhysicsAssetEditor::AddViewportOverlayWidget(TSharedRef<SWidget> InOverlaidWidget, int32 ZOrder)
+{
+	if (Viewport.IsValid())
+	{
+		Viewport->AddOverlayWidget(InOverlaidWidget, ZOrder);
+	}
+}
+
+void FPhysicsAssetEditor::RemoveViewportOverlayWidget(TSharedRef<SWidget> InOverlaidWidget)
+{
+	if (Viewport.IsValid())
+	{
+		Viewport->RemoveOverlayWidget(InOverlaidWidget);
+	}
+}
+
 void FPhysicsAssetEditor::HandleViewportSelectionChanged(const TArray<FPhysicsAssetEditorSelectedElement>& InSelectedElements)
 {
 	if (!bSelecting)
diff --git a/Engine/Source/Editor/PhysicsAssetEditor/Private/PhysicsAssetEditor.h b/Engine/Source/Editor/PhysicsAssetEditor/Private/PhysicsAssetEditor.h
--- a/Engine/Source/Editor/PhysicsAssetEditor/Private/PhysicsAssetEditor.h
+++ b/Engine/Source/Editor/PhysicsAssetEditor/Private/PhysicsAssetEditor.h
@@ -111,6 +111,13 @@ public:
 	virtual TSharedRef<IPersonaToolkit> GetPersonaToolkit() const override { return PersonaToolkit.ToSharedRef(); }
 	virtual void OnClose() override;
 
+	// FAssetEditorToolkit interface
+	virtual void AddViewportOverlayWidget(TSharedRef<SWidget> InOverlaidWidget, int32 ZOrder = INDEX_NONE) override;
+	virtual void RemoveViewportOverlayWidget(TSharedRef<SWidget> InOverlaidWidget) override;
+
+	/** Store the viewport reference for overlay widget forwarding. */
+	void HandleViewportCreated(const TSharedRef<class IPersonaViewport>& InViewport);
+
 	/** FGCObject interface */
 	virtual void AddReferencedObjects(FReferenceCollector& Collector) override;
 	virtual FString GetReferencerName() const override
@@ -429,6 +436,9 @@ private:
 	/** The persona toolkit */
 	TSharedPtr<IPersonaToolkit> PersonaToolkit;
 
+	/** Viewport for overlay widget forwarding */
+	TSharedPtr<class IPersonaViewport> Viewport;
+
 	/** The current physics asset graph, if any */
 	TWeakPtr<SPhysicsAssetGraph> PhysicsAssetGraph;
 
diff --git a/Engine/Source/Editor/PhysicsAssetEditor/Private/PhysicsAssetEditorMode.cpp b/Engine/Source/Editor/PhysicsAssetEditor/Private/PhysicsAssetEditorMode.cpp
--- a/Engine/Source/Editor/PhysicsAssetEditor/Private/PhysicsAssetEditorMode.cpp
+++ b/Engine/Source/Editor/PhysicsAssetEditor/Private/PhysicsAssetEditorMode.cpp
@@ -1,6 +1,8 @@
 // Copyright Epic Games, Inc. All Rights Reserved.
 
 #include "PhysicsAssetEditorMode.h"
+
+#include "EditorViewportClient.h"
 #include "Engine/SkeletalMesh.h"
 #include "PhysicsAssetEditor.h"
 #include "ISkeletonTree.h"
@@ -186,6 +188,12 @@ FPhysicsAssetEditorMode::FPhysicsAssetEditorMode(TSharedRef<FWorkflowCentricAppl
 		InViewport->GetPinnedCommandList()->BindCommandList(PhysicsAssetEditorPtr.Pin()->GetViewportCommandList().ToSharedRef());
 		InViewport->GetPinnedCommandList()->RegisterCustomWidget(IPinnedCommandList::FOnGenerateCustomWidget::CreateSP(PhysicsAssetEditorPtr.Pin().Get(), &FPhysicsAssetEditor::MakeConstraintScaleWidget), TEXT("ConstraintScaleWidget"), LOCTEXT("ConstraintScaleLabel", "Constraint Scale"));
 		InViewport->GetPinnedCommandList()->RegisterCustomWidget(IPinnedCommandList::FOnGenerateCustomWidget::CreateSP(PhysicsAssetEditorPtr.Pin().Get(), &FPhysicsAssetEditor::MakeCollisionOpacityWidget), TEXT("CollisionOpacityWidget"), LOCTEXT("CollisionOpacityLabel", "Collision Opacity"));
+
+		// Store viewport reference on the editor for overlay widget forwarding.
+		PhysicsAssetEditorPtr.Pin()->HandleViewportCreated(InViewport);
+
+		// we need the viewport client to start out focused, or else it won't get ticked until we click inside it.
+		InViewport->GetViewportClient().ReceivedFocus(InViewport->GetViewportClient().Viewport);
 	});
 	ViewportArgs.OnGetViewportText = FOnGetViewportText::CreateLambda([this](EViewportCorner InViewportCorner)
 	{
diff --git a/Engine/Source/Editor/SkeletonEditor/Private/SkeletonEditor.cpp b/Engine/Source/Editor/SkeletonEditor/Private/SkeletonEditor.cpp
--- a/Engine/Source/Editor/SkeletonEditor/Private/SkeletonEditor.cpp
+++ b/Engine/Source/Editor/SkeletonEditor/Private/SkeletonEditor.cpp
@@ -24,10 +24,14 @@
 #include "IEditableSkeleton.h"
 #include "ISkeletonTreeItem.h"
 #include "Algo/Transform.h"
+#include "EditorModeManager.h"
+#include "EditorViewportClient.h"
+#include "IPersonaViewport.h"
 #include "PersonaToolMenuContext.h"
 #include "ToolMenus.h"
 #include "ToolMenuMisc.h"
 #include "SkeletonToolMenuContext.h"
+#include "Tools/EdModeInteractiveToolsContext.h"
 
 const FName SkeletonEditorAppIdentifier = FName(TEXT("SkeletonEditorApp"));
 
@@ -409,6 +413,30 @@ void FSkeletonEditor::HandleAnimationSequenceBrowserCreated(const TSharedRef<IAn
 	SequenceBrowser = InSequenceBrowser;
 }
 
+void FSkeletonEditor::HandleViewportCreated(const TSharedRef<IPersonaViewport>& InViewport)
+{
+	Viewport = InViewport;
+
+	// we need the viewport client to start out focused, or else it won't get ticked until we click inside it.
+	InViewport->GetViewportClient().ReceivedFocus(InViewport->GetViewportClient().Viewport);
+}
+
+void FSkeletonEditor::AddViewportOverlayWidget(TSharedRef<SWidget> InOverlaidWidget, int32 ZOrder)
+{
+	if (Viewport.IsValid())
+	{
+		Viewport->AddOverlayWidget(InOverlaidWidget, ZOrder);
+	}
+}
+
+void FSkeletonEditor::RemoveViewportOverlayWidget(TSharedRef<SWidget> InOverlaidWidget)
+{
+	if (Viewport.IsValid())
+	{
+		Viewport->RemoveOverlayWidget(InOverlaidWidget);
+	}
+}
+
 IAnimationSequenceBrowser* FSkeletonEditor::GetAssetBrowser() const
 {
 	return SequenceBrowser.Pin().Get();
diff --git a/Engine/Source/Editor/SkeletonEditor/Private/SkeletonEditor.h b/Engine/Source/Editor/SkeletonEditor/Private/SkeletonEditor.h
--- a/Engine/Source/Editor/SkeletonEditor/Private/SkeletonEditor.h
+++ b/Engine/Source/Editor/SkeletonEditor/Private/SkeletonEditor.h
@@ -100,6 +100,13 @@ public:
 	void HandleDetailsCreated(const TSharedRef<class IDetailsView>& InDetailsView);
 
 	void HandleAnimationSequenceBrowserCreated(const TSharedRef<class IAnimationSequenceBrowser>& InSequenceBrowser);
+
+	void HandleViewportCreated(const TSharedRef<IPersonaViewport>& InViewport);
+
+	// FAssetEditorToolkit interface
+	virtual void AddViewportOverlayWidget(TSharedRef<SWidget> InOverlaidWidget, int32 ZOrder = INDEX_NONE) override;
+	virtual void RemoveViewportOverlayWidget(TSharedRef<SWidget> InOverlaidWidget) override;
+
 private:
 	void ExtendMenu();
 
diff --git a/Engine/Source/Editor/SkeletonEditor/Private/SkeletonEditorMode.cpp b/Engine/Source/Editor/SkeletonEditor/Private/SkeletonEditorMode.cpp
--- a/Engine/Source/Editor/SkeletonEditor/Private/SkeletonEditorMode.cpp
+++ b/Engine/Source/Editor/SkeletonEditor/Private/SkeletonEditorMode.cpp
@@ -28,6 +28,7 @@ FSkeletonEditorMode::FSkeletonEditorMode(TSharedRef<FWorkflowCentricApplication>
 
 	FPersonaViewportArgs ViewportArgs(SkeletonEditor->GetPersonaToolkit()->GetPreviewScene());
 	ViewportArgs.ContextName = TEXT("SkeletonEditor.Viewport");
+	ViewportArgs.OnViewportCreated = FOnViewportCreated::CreateSP(&SkeletonEditor.Get(), &FSkeletonEditor::HandleViewportCreated);
 
 	PersonaModule.RegisterPersonaViewportTabFactories(TabFactories, InHostingApp, ViewportArgs);
 
