[UE] Physics Asset Editor Python API Support

- Minimal engine changes to support the creation of a new python API for physics assets and the Physics Asset Editor mode to be implemented as a project side plugin
- Extend PhAT selection system to allow filtering selection to specific primitive types
- Create a unified selection interface that includes bones, bodies, constraints and primitives and can be called through the public interface class

(reviewed as part of 51826649)

#Synced-CL 52099373

diff --git a/Engine/Source/Developer/PhysicsUtilities/Public/PhysicsAssetUtils.h b/Engine/Source/Developer/PhysicsUtilities/Public/PhysicsAssetUtils.h
--- a/Engine/Source/Developer/PhysicsUtilities/Public/PhysicsAssetUtils.h
+++ b/Engine/Source/Developer/PhysicsUtilities/Public/PhysicsAssetUtils.h
@@ -282,7 +282,7 @@ namespace FPhysicsAssetUtils
 	 */
 	template<typename RealType>
 	FMatrix ComputeViewCovarianceMatrix(TArrayView<const UE::Math::TVector<RealType>, int> PointView);
-	PHYSICSUTILITIES_API FMatrix ComputeCovarianceMatrix(const FBoneVertInfo& VertInfo);
+	PHYSICSUTILITIES_API FMatrix ComputeCovarianceMatrix(const TArray<FVector3f>& Verts);
 
 	extern template FMatrix ComputeViewCovarianceMatrix(TArrayView<const UE::Math::TVector<float>, int> PointView);
 	extern template FMatrix ComputeViewCovarianceMatrix(TArrayView<const UE::Math::TVector<double>, int> PointView);
diff --git a/Engine/Source/Editor/PhysicsAssetEditor/Private/PhysicsAssetEditor.cpp b/Engine/Source/Editor/PhysicsAssetEditor/Private/PhysicsAssetEditor.cpp
--- a/Engine/Source/Editor/PhysicsAssetEditor/Private/PhysicsAssetEditor.cpp
+++ b/Engine/Source/Editor/PhysicsAssetEditor/Private/PhysicsAssetEditor.cpp
@@ -4889,4 +4889,211 @@ TSharedRef<SWidget> FPhysicsAssetEditor::MakeCollisionOpacityWidget()
 		];
 }
 
+TArray<int32> FPhysicsAssetEditor::GetSelectedBodyIndices() const
+{
+	TArray<int32> Result;
+	const FPhysicsAssetEditorSharedData::SelectionUniqueRange SelectionRange = SharedData->UniqueSelectionReferencingBodies();
+	Result.Reserve(SelectionRange.Num());
+	Algo::Transform(SelectionRange, Result, [](const FPhysicsAssetEditorSharedData::FSelection& Selection) { return Selection.Index; });
+	return Result;
+}
+
+TArray<int32> FPhysicsAssetEditor::GetSelectedConstraintIndices() const
+{
+	TArray<int32> Result;
+	const FPhysicsAssetEditorSharedData::SelectionFilterRange SelectionRange = SharedData->SelectedConstraints();
+	Result.Reserve(SelectionRange.Num());
+	Algo::Transform(SelectionRange, Result, [](const FPhysicsAssetEditorSharedData::FSelection& Selection) { return Selection.Index; });
+	return Result;
+}
+
+namespace PhysicsAssetEditorSelectionPrivate
+{
+	/** Map a public EPhysicsAssetEditorSelection primitive flag to an EAggCollisionShape value.
+	 *  Returns EAggCollisionShape::Unknown if the flag is not a single primitive type. */
+	EAggCollisionShape::Type SelectionTypeToShapeType(const EPhysicsAssetEditorSelection InType)
+	{
+		switch (InType)
+		{
+		case EPhysicsAssetEditorSelection::PrimitiveSphere:          return EAggCollisionShape::Sphere;
+		case EPhysicsAssetEditorSelection::PrimitiveBox:             return EAggCollisionShape::Box;
+		case EPhysicsAssetEditorSelection::PrimitiveCapsule:         return EAggCollisionShape::Sphyl;
+		case EPhysicsAssetEditorSelection::PrimitiveConvex:          return EAggCollisionShape::Convex;
+		case EPhysicsAssetEditorSelection::PrimitiveTaperedCapsule:  return EAggCollisionShape::TaperedCapsule;
+		case EPhysicsAssetEditorSelection::PrimitiveLevelSet:        return EAggCollisionShape::LevelSet;
+		case EPhysicsAssetEditorSelection::PrimitiveSkinnedLevelSet: return EAggCollisionShape::SkinnedLevelSet;
+		default: return EAggCollisionShape::Unknown;
+		}
+	}
+
+	/** Check whether a given EPhysicsAssetEditorSelection type flag includes the given shape type. */
+	bool SelectionTypeIncludesShape(const int32 InSelectionFlags, const EAggCollisionShape::Type InShapeType)
+	{
+		switch (InShapeType)
+		{
+		case EAggCollisionShape::Sphere:          return (InSelectionFlags & UnderlyingType(EPhysicsAssetEditorSelection::PrimitiveSphere)) != 0;
+		case EAggCollisionShape::Box:             return (InSelectionFlags & UnderlyingType(EPhysicsAssetEditorSelection::PrimitiveBox)) != 0;
+		case EAggCollisionShape::Sphyl:           return (InSelectionFlags & UnderlyingType(EPhysicsAssetEditorSelection::PrimitiveCapsule)) != 0;
+		case EAggCollisionShape::Convex:          return (InSelectionFlags & UnderlyingType(EPhysicsAssetEditorSelection::PrimitiveConvex)) != 0;
+		case EAggCollisionShape::TaperedCapsule:  return (InSelectionFlags & UnderlyingType(EPhysicsAssetEditorSelection::PrimitiveTaperedCapsule)) != 0;
+		case EAggCollisionShape::LevelSet:        return (InSelectionFlags & UnderlyingType(EPhysicsAssetEditorSelection::PrimitiveLevelSet)) != 0;
+		case EAggCollisionShape::SkinnedLevelSet: return (InSelectionFlags & UnderlyingType(EPhysicsAssetEditorSelection::PrimitiveSkinnedLevelSet)) != 0;
+		default: return false;
+		}
+	}
+}
+
+void FPhysicsAssetEditor::ModifySelection(EPhysicsAssetEditorSelection ElementType,
+	const TArray<int32>& Indices, bool bAddToSelection)
+{
+	if (Indices.Num() == 0)
+	{
+		return;
+	}
+
+	FScopedBulkSelection BulkSelection(SharedData);
+
+	if (ElementType == EPhysicsAssetEditorSelection::Body)
+	{
+		SharedData->SetSelectedBodiesAllPrimitive(Indices, bAddToSelection);
+	}
+	else if (ElementType == EPhysicsAssetEditorSelection::Constraint)
+	{
+		SharedData->ModifySelectedConstraints(Indices, bAddToSelection);
+	}
+	else if (ElementType == EPhysicsAssetEditorSelection::CenterOfMass)
+	{
+		TArray<FPhysicsAssetEditorSelectedElement> CoMSelections;
+		CoMSelections.Reserve(Indices.Num());
+		for (const int32 BodyIndex : Indices)
+		{
+			CoMSelections.Add(MakeCoMSelection(BodyIndex));
+		}
+		SharedData->ModifySelectedCoMs(CoMSelections, bAddToSelection);
+	}
+	else if (ElementType == EPhysicsAssetEditorSelection::Bone)
+	{
+		USkeletalMesh* PreviewMesh = GetPersonaToolkit()->GetPreviewScene()->GetPreviewMesh();
+		if (!PreviewMesh)
+		{
+			return;
+		}
+
+		const FReferenceSkeleton& RefSkeleton = PreviewMesh->GetRefSkeleton();
+		TArray<FName> BoneNames;
+		BoneNames.Reserve(Indices.Num());
+		for (const int32 BoneIndex : Indices)
+		{
+			if (RefSkeleton.IsValidIndex(BoneIndex))
+			{
+				BoneNames.Add(RefSkeleton.GetBoneName(BoneIndex));
+			}
+		}
+
+		if (BoneNames.Num() > 0)
+		{
+			if (bAddToSelection)
+			{
+				// Merge new names with existing selected bones
+				UDebugSkelMeshComponent* MeshComp = GetPersonaToolkit()->GetPreviewScene()->GetPreviewMeshComponent();
+				if (MeshComp)
+				{
+					for (const int32 ExistingBoneIdx : MeshComp->BonesOfInterest)
+					{
+						if(RefSkeleton.IsValidIndex(ExistingBoneIdx))
+						{
+							FName ExistingName = RefSkeleton.GetBoneName(ExistingBoneIdx);
+							BoneNames.AddUnique(ExistingName);
+						}
+					}
+				}
+				GetPersonaToolkit()->GetPreviewScene()->SetSelectedBones(BoneNames, ESelectInfo::Direct);
+			}
+			else
+			{
+				// Remove specified bones from current selection
+				UDebugSkelMeshComponent* MeshComp = GetPersonaToolkit()->GetPreviewScene()->GetPreviewMeshComponent();
+				if (MeshComp)
+				{
+					TSet<int32> BoneIndicesToRemove(Indices);
+					TArray<FName> RemainingBones;
+					for (const int32 ExistingBoneIdx : MeshComp->BonesOfInterest)
+					{
+						if (!BoneIndicesToRemove.Contains(ExistingBoneIdx) && RefSkeleton.IsValidIndex(ExistingBoneIdx))
+						{
+							RemainingBones.Add(RefSkeleton.GetBoneName(ExistingBoneIdx));
+						}
+					}
+					if (RemainingBones.Num() > 0)
+					{
+						GetPersonaToolkit()->GetPreviewScene()->SetSelectedBones(RemainingBones, ESelectInfo::Direct);
+					}
+					else
+					{
+						GetPersonaToolkit()->GetPreviewScene()->ClearSelectedBone();
+					}
+				}
+			}
+		}
+	}
+	else if ((ElementType == EPhysicsAssetEditorSelection::Primitive) || (ElementType == EPhysicsAssetEditorSelection::AllPrimitive))
+	{
+		// Select all primitives in the given bodies
+		SharedData->SetSelectedBodiesAllPrimitive(Indices, bAddToSelection);
+	}
+	else if (UnderlyingType(ElementType) & UnderlyingType(EPhysicsAssetEditorSelection::AllPrimitive))
+	{
+		// Select only primitives of the matching shape type(s)
+		const int32 SelectionFlags = UnderlyingType(ElementType);
+		UPhysicsAsset* PhysicsAsset = SharedData->PhysicsAsset;
+		TArray<FPhysicsAssetEditorSelectedElement> PrimSelections;
+
+		for (const int32 BodyIndex : Indices)
+		{
+			if (!PhysicsAsset->SkeletalBodySetups.IsValidIndex(BodyIndex))
+			{
+				continue;
+			}
+
+			UBodySetup* BodySetup = PhysicsAsset->SkeletalBodySetups[BodyIndex];
+			if (!BodySetup)
+			{
+				continue;
+			}
+
+			auto AddPrimitivesOfType = [&PrimSelections, SelectionFlags, BodyIndex](const auto& ShapeElems)
+			{
+				for (int32 PrimIdx = 0; PrimIdx < ShapeElems.Num(); ++PrimIdx)
+				{
+					if (PhysicsAssetEditorSelectionPrivate::SelectionTypeIncludesShape(SelectionFlags, ShapeElems[PrimIdx].GetShapeType()))
+					{
+						PrimSelections.Add(MakePrimitiveSelection(BodyIndex, ShapeElems[PrimIdx].GetShapeType(), PrimIdx));
+					}
+				}
+			};
+
+			const FKAggregateGeom& AggGeom = BodySetup->AggGeom;
+
+			AddPrimitivesOfType(AggGeom.SphereElems);
+			AddPrimitivesOfType(AggGeom.BoxElems);
+			AddPrimitivesOfType(AggGeom.SphylElems);
+			AddPrimitivesOfType(AggGeom.ConvexElems);
+			AddPrimitivesOfType(AggGeom.TaperedCapsuleElems);
+			AddPrimitivesOfType(AggGeom.LevelSetElems);
+			AddPrimitivesOfType(AggGeom.SkinnedLevelSetElems);
+		}
+
+		if (PrimSelections.Num() > 0)
+		{
+			SharedData->ModifySelectedPrimitives(PrimSelections, bAddToSelection);
+		}	
+	}
+}
+
+void FPhysicsAssetEditor::ClearSelection()
+{
+	SharedData->ClearSelected();
+	GetPersonaToolkit()->GetPreviewScene()->ClearSelectedBone();
+}
+
 #undef LOCTEXT_NAMESPACE
diff --git a/Engine/Source/Editor/PhysicsAssetEditor/Private/PhysicsAssetEditor.h b/Engine/Source/Editor/PhysicsAssetEditor/Private/PhysicsAssetEditor.h
--- a/Engine/Source/Editor/PhysicsAssetEditor/Private/PhysicsAssetEditor.h
+++ b/Engine/Source/Editor/PhysicsAssetEditor/Private/PhysicsAssetEditor.h
@@ -107,6 +107,13 @@ public:
 		return FString(TEXT("Engine/Physics/PhysicsAssetEditor"));
 	}
 
+	/** IPhysicsAssetEditor interface */
+	virtual TArray<int32> GetSelectedBodyIndices() const override;
+	virtual TArray<int32> GetSelectedConstraintIndices() const override;
+
+	virtual void ModifySelection(EPhysicsAssetEditorSelection ElementType, const TArray<int32>& Indices, bool bAddToSelection) override;
+	virtual void ClearSelection() override;
+
 	/** IHasPersonaToolkit interface */
 	virtual TSharedRef<IPersonaToolkit> GetPersonaToolkit() const override { return PersonaToolkit.ToSharedRef(); }
 	virtual void OnClose() override;
diff --git a/Engine/Source/Editor/PhysicsAssetEditor/Private/PhysicsAssetEditorSelection.h b/Engine/Source/Editor/PhysicsAssetEditor/Private/PhysicsAssetEditorSelection.h
--- a/Engine/Source/Editor/PhysicsAssetEditor/Private/PhysicsAssetEditorSelection.h
+++ b/Engine/Source/Editor/PhysicsAssetEditor/Private/PhysicsAssetEditorSelection.h
@@ -3,6 +3,7 @@
 #pragma once
 
 #include "CoreMinimal.h"
+#include "IPhysicsAssetEditor.h"
 #include "PhysicsEngine/ShapeElem.h"
 #include "Preferences/PhysicsAssetEditorOptions.h"
 
@@ -17,17 +18,15 @@ struct FPhysicsAssetEditorSelectedElement
 {
 	GENERATED_BODY()
 
-	enum Type : int32
-	{
-		None,
-
-		Body			= 1,		// A physics body, comprised of one or more primitives.
-		CenterOfMass	= 1 << 1,	// A Center of Mass of a physics body.
-		Constraint		= 1 << 2,	// A constraint between two physics bodies.
-		Primitive		= 1 << 3,	// A primitive that defines the geometry of a physics body.
+	using Type = std::underlying_type_t<EPhysicsAssetEditorSelection>;
 
-		All				= ~None
-	};
+	// Backward-compatability aliases — existing code uses FPhysicsAssetEditorSelectedElement::Body etc.
+	static constexpr Type None         = UnderlyingType(EPhysicsAssetEditorSelection::None);
+	static constexpr Type Body         = UnderlyingType(EPhysicsAssetEditorSelection::Body);
+	static constexpr Type CenterOfMass = UnderlyingType(EPhysicsAssetEditorSelection::CenterOfMass);
+	static constexpr Type Constraint   = UnderlyingType(EPhysicsAssetEditorSelection::Constraint);
+	static constexpr Type Primitive    = UnderlyingType(EPhysicsAssetEditorSelection::Primitive);
+	static constexpr Type All          = UnderlyingType(EPhysicsAssetEditorSelection::All);
 
 	FPhysicsAssetEditorSelectedElement() = default;
 	FPhysicsAssetEditorSelectedElement(const Type InSelectedElementType, const int32 InBodyIndex, const EAggCollisionShape::Type InPrimitiveType, const int32 InPrimitiveIndex);
@@ -50,7 +49,7 @@ struct FPhysicsAssetEditorSelectedElement
 	int32 PrimitiveIndex = INDEX_NONE;
 
 	UPROPERTY()
-	int32 SelectedElementType = Type::None;
+	int32 SelectedElementType = None;
 
 };
 
diff --git a/Engine/Source/Editor/PhysicsAssetEditor/Public/IPhysicsAssetEditor.h b/Engine/Source/Editor/PhysicsAssetEditor/Public/IPhysicsAssetEditor.h
--- a/Engine/Source/Editor/PhysicsAssetEditor/Public/IPhysicsAssetEditor.h
+++ b/Engine/Source/Editor/PhysicsAssetEditor/Public/IPhysicsAssetEditor.h
@@ -7,6 +7,7 @@
 #include "IHasPersonaToolkit.h"
 #include "Features/IModularFeature.h"
 #include "AssetDefinition.h"
+#include "PhysicsAssetEditorSelectionType.h"
 
 class UPhysicsAsset;
 
@@ -16,7 +17,23 @@ class UPhysicsAsset;
 
 class IPhysicsAssetEditor : public FPersonaAssetEditorToolkit, public IHasPersonaToolkit
 {
-
+public:
+	/** Get the body indices of the currently selected bodies. */
+	virtual TArray<int32> GetSelectedBodyIndices() const = 0;
+
+	/** Get the constraint indices currently selected. */
+	virtual TArray<int32> GetSelectedConstraintIndices() const = 0;
+
+	/** Add or remove elements from the selection.
+	 *  For Body/CenterOfMass: Indices are body indices.
+	 *  For Bone: Indices are skeleton bone indices.
+	 *  For Constraint: Indices are constraint indices.
+	 *  For primitive types (PrimitiveCapsule, Primitive, etc.): Indices are body indices;
+	 *    the specific enum value determines which shape types are selected within those bodies. */
+	virtual void ModifySelection(EPhysicsAssetEditorSelection ElementType, const TArray<int32>& Indices, bool bAddToSelection) = 0;
+
+	/** Clear all selection (bodies, constraints, primitives, center of mass, bones). */
+	virtual void ClearSelection() = 0;
 };
 
 class IPhysicsAssetEditorOverride : public IModularFeature
diff --git a/Engine/Source/Editor/PhysicsAssetEditor/Public/PhysicsAssetEditorSelectionType.h b/Engine/Source/Editor/PhysicsAssetEditor/Public/PhysicsAssetEditorSelectionType.h
new file mode 100644
--- /dev/null
+++ b/Engine/Source/Editor/PhysicsAssetEditor/Public/PhysicsAssetEditorSelectionType.h
@@ -0,0 +1,39 @@
+// Copyright Epic Games, Inc. All Rights Reserved.
+
+#pragma once
+
+#include "CoreMinimal.h"
+
+/** Element types in the Physics Asset Editor selection system.
+ *  
+ *  Individual primitive types allow shape-specific selection via ModifySelection. 
+ */
+enum class EPhysicsAssetEditorSelection : int32
+{
+	None = 0,
+	Body = 1 << 0,
+	CenterOfMass = 1 << 1,
+	Constraint = 1 << 2,
+	Primitive = 1 << 3,
+	Bone = 1 << 4,
+
+	// Individual primitive types (for shape-specific selection)
+	PrimitiveSphere = 1 << 5,
+	PrimitiveBox = 1 << 6,
+	PrimitiveCapsule = 1 << 7,   // Maps to EAggCollisionShape::Sphyl
+	PrimitiveConvex = 1 << 8,
+	PrimitiveTaperedCapsule = 1 << 9,
+	PrimitiveLevelSet = 1 << 10,
+	PrimitiveSkinnedLevelSet = 1 << 11,
+
+	// Combined: all primitive types
+	AllPrimitive = PrimitiveSphere | PrimitiveBox | PrimitiveCapsule | PrimitiveConvex
+	| PrimitiveTaperedCapsule | PrimitiveLevelSet | PrimitiveSkinnedLevelSet,
+
+	All = ~None
+};
+
+static constexpr std::underlying_type_t<EPhysicsAssetEditorSelection> UnderlyingType(const EPhysicsAssetEditorSelection SelectionType)
+{
+	return static_cast<std::underlying_type_t<EPhysicsAssetEditorSelection>>(SelectionType);
+}
\ No newline at end of file
