[Physics Editor] Flag the physics editor as requiring legacy viewport interactions.

#rnx
#Synced-CL 52132857

diff --git a/Engine/Source/Editor/PhysicsAssetEditor/Private/PhysicsAssetEditorEditMode.h b/Engine/Source/Editor/PhysicsAssetEditor/Private/PhysicsAssetEditorEditMode.h
--- a/Engine/Source/Editor/PhysicsAssetEditor/Private/PhysicsAssetEditorEditMode.h
+++ b/Engine/Source/Editor/PhysicsAssetEditor/Private/PhysicsAssetEditorEditMode.h
@@ -52,6 +52,7 @@ public:
 	virtual bool IsCompatibleWith(FEditorModeID OtherModeID) const override { return true; }
 	virtual bool ReceivedFocus(FEditorViewportClient* ViewportClient, FViewport* Viewport);
 	virtual bool LostFocus(FEditorViewportClient* ViewportClient, FViewport* Viewport);
+	virtual bool RequiresLegacyViewportInteractions() const override { return true; }
 
 	// Start IGizmoEdModeInterface overrides
 	virtual bool BeginTransform(const FGizmoState& InState) override;
