PhysicsAssetEditor ITF Support, with custom physics simulation interactions to match legacy

#Synced-CL 52474860

diff --git a/Engine/Source/Editor/Experimental/EditorInteractiveToolsFramework/Private/EditorGizmos/EditorTRSGizmo.cpp b/Engine/Source/Editor/Experimental/EditorInteractiveToolsFramework/Private/EditorGizmos/EditorTRSGizmo.cpp
--- a/Engine/Source/Editor/Experimental/EditorInteractiveToolsFramework/Private/EditorGizmos/EditorTRSGizmo.cpp
+++ b/Engine/Source/Editor/Experimental/EditorInteractiveToolsFramework/Private/EditorGizmos/EditorTRSGizmo.cpp
@@ -1413,16 +1413,6 @@ void UEditorTRSGizmo::SetupBehaviors()
 		)
 	);
 
-	MouseBehavior->ModifierCheckFunc = [this](const FInputDeviceState& InputDeviceState)
-	{
-		// Filter out Alt unless we need it for duplication
-		if (FInputDeviceState::IsAltKeyDown(InputDeviceState))
-		{
-			return AllowsDragDuplicate();
-		}
-
-		return true;
-	};
 	AddInputBehavior(MouseBehavior);
 	ClickAndDragBehaviorWeak = MouseBehavior;
 
diff --git a/Engine/Source/Editor/PhysicsAssetEditor/PhysicsAssetEditor.Build.cs b/Engine/Source/Editor/PhysicsAssetEditor/PhysicsAssetEditor.Build.cs
--- a/Engine/Source/Editor/PhysicsAssetEditor/PhysicsAssetEditor.Build.cs
+++ b/Engine/Source/Editor/PhysicsAssetEditor/PhysicsAssetEditor.Build.cs
@@ -42,7 +42,9 @@ public class PhysicsAssetEditor : ModuleRules
 				"ApplicationCore",
 				"EditorStyle",
 				"ToolWidgets",
-				"TypedElementRuntime"
+				"TypedElementRuntime",
+				"EditorInteractiveToolsFramework",
+				"InteractiveToolsFramework"
 			}
         );
 		SetupModulePhysicsSupport(Target);
diff --git a/Engine/Source/Editor/PhysicsAssetEditor/Private/PhysicsAssetEditor.cpp b/Engine/Source/Editor/PhysicsAssetEditor/Private/PhysicsAssetEditor.cpp
--- a/Engine/Source/Editor/PhysicsAssetEditor/Private/PhysicsAssetEditor.cpp
+++ b/Engine/Source/Editor/PhysicsAssetEditor/Private/PhysicsAssetEditor.cpp
@@ -106,6 +106,11 @@ DEFINE_LOG_CATEGORY(LogPhysicsAssetEditor);
 
 namespace PhysicsAssetEditor
 {
+	TAutoConsoleVariable<bool> CVarUseITFViewportInteractions(
+		TEXT("PhysicsAssetEditor.UseITFViewportInteractions"),
+		false,
+		TEXT("If true, enables ITF based drag interactions, camera controls, and gizmos in the physics asset editor."));
+
 	const float	DefaultPrimSize = 15.0f;
 	const float	DuplicateXOffset = 10.0f;
 
@@ -186,6 +191,8 @@ FPhysicsAssetEditor::~FPhysicsAssetEditor()
 		ImpToggleSimulation();
 	}
 
+	PhysicsAssetEditor::CVarUseITFViewportInteractions->OnChangedDelegate().RemoveAll(this);
+
 	GEditor->UnregisterForUndo(this);
 	GEditor->GetEditorSubsystem<UImportSubsystem>()->OnAssetReimport.Remove(OnAssetReimportDelegateHandle);
 	if (PersonaToolkit.IsValid())
@@ -298,6 +305,28 @@ void FPhysicsAssetEditor::InitPhysicsAssetEditor(const EToolkitMode::Type Mode,
 	const bool bCreateDefaultToolbar = true;
 	FAssetEditorToolkit::InitAssetEditor(Mode, InitToolkitHost, PhysicsAssetEditorAppIdentifier, FTabManager::FLayout::NullLayout, bCreateDefaultStandaloneMenu, bCreateDefaultToolbar, ObjectToEdit);
 
+	PhysicsAssetEditor::CVarUseITFViewportInteractions->OnChangedDelegate().AddSPLambda(this, [ModeTools = TWeakPtr<FEditorModeTools>(GetEditorModeManager().AsShared())](IConsoleVariable*)
+	{
+		if (TSharedPtr<FEditorModeTools> ModeToolsPin = ModeTools.Pin())
+		{
+			const bool bEnabled = PhysicsAssetEditor::CVarUseITFViewportInteractions->GetBool();
+			if (UModeManagerInteractiveToolsContext* ToolsContext = ModeToolsPin->GetInteractiveToolsContext())
+			{
+				ToolsContext->SetViewportInteractionsEnabled(bEnabled);
+			}
+			ModeToolsPin->SetSupportsViewportITF(bEnabled);
+		}
+	});
+
+	{
+		const bool bEnabled = PhysicsAssetEditor::CVarUseITFViewportInteractions->GetBool();
+		if (UModeManagerInteractiveToolsContext* ToolsContext = GetEditorModeManager().GetInteractiveToolsContext())
+		{
+			ToolsContext->SetViewportInteractionsEnabled(bEnabled);
+		}
+		GetEditorModeManager().SetSupportsViewportITF(bEnabled);
+	}
+
 	AddApplicationMode(
 		PhysicsAssetEditorModes::PhysicsAssetEditorMode,
 		MakeShareable(new FPhysicsAssetEditorMode(SharedThis(this), SkeletonTree.ToSharedRef(), PersonaToolkit->GetPreviewScene())));
diff --git a/Engine/Source/Editor/PhysicsAssetEditor/Private/PhysicsAssetEditorEditMode.cpp b/Engine/Source/Editor/PhysicsAssetEditor/Private/PhysicsAssetEditorEditMode.cpp
--- a/Engine/Source/Editor/PhysicsAssetEditor/Private/PhysicsAssetEditorEditMode.cpp
+++ b/Engine/Source/Editor/PhysicsAssetEditor/Private/PhysicsAssetEditorEditMode.cpp
@@ -27,6 +27,9 @@
 
 #include "SceneView.h"
 #include "SnappingUtils.h"
+#include "Tools/EdModeInteractiveToolsContext.h"
+#include "ViewportInteractions/PhysicsAssetSimInteraction.h"
+#include "ViewportInteractions/ViewportInteractionsBehaviorSource.h"
 
 #define LOCTEXT_NAMESPACE "PhysicsAssetEditorEditMode"
 
@@ -50,6 +53,31 @@ FPhysicsAssetEditorEditMode::FPhysicsAssetEditorEditMode()
 	check(PhysicsAssetEditorFont);
 }
 
+void FPhysicsAssetEditorEditMode::Enter()
+{
+	IPersonaEditMode::Enter();
+
+	if (UModeManagerInteractiveToolsContext* Context = GetModeManager()->GetInteractiveToolsContext())
+	{
+		Context->SetViewportInteractionsEnabled(true);
+		Context->OnBuildViewportInteractions().AddSPLambda(this, [this](UViewportInteractionsBehaviorSource* Source)
+		{
+			UPhysicsAssetSimInteraction* NewSimInteraction = Source->AddInteraction<UPhysicsAssetSimInteraction>(true);
+			NewSimInteraction->SetEditMode(this);
+		});
+	}
+}
+
+void FPhysicsAssetEditorEditMode::Exit()
+{
+	if (UModeManagerInteractiveToolsContext* Context = GetModeManager()->GetInteractiveToolsContext())
+	{
+		Context->OnBuildViewportInteractions().RemoveAll(this);
+	}
+
+	IPersonaEditMode::Exit();
+}
+
 bool FPhysicsAssetEditorEditMode::GetCameraTarget(FSphere& OutTarget) const
 {
 	bool bHandled = false;
@@ -317,17 +345,22 @@ bool FPhysicsAssetEditorEditMode::InputAxis(FEditorViewportClient* InViewportCli
 	// If we are 'manipulating' don't move the camera but do something else with mouse input.
 	if (SharedData->IsManipulating())
 	{
-		bool bCtrlDown = InViewport->KeyState(EKeys::LeftControl) || InViewport->KeyState(EKeys::RightControl);
-
+		// Simulation mouse movement is handled by UPhysicsAssetSimInteraction::OnDragDelta
+		// when viewport interactions are active (Gizmos.Version > 0). Fall back to the
+		// legacy path when ITF is not available.
 		if (SharedData->bRunningSimulation)
 		{
-			if (Key == EKeys::MouseX)
-			{
-				SimMouseMove(InViewportClient, Delta, 0.0f);
-			}
-			else if (Key == EKeys::MouseY)
+			const UModeManagerInteractiveToolsContext* Context = GetModeManager()->GetInteractiveToolsContext();
+			if (!Context || !Context->UsesViewportInteractions())
 			{
-				SimMouseMove(InViewportClient, 0.0f, Delta);
+				if (Key == EKeys::MouseX)
+				{
+					SimMouseMove(InViewportClient, Delta, 0.0f);
+				}
+				else if (Key == EKeys::MouseY)
+				{
+					SimMouseMove(InViewportClient, 0.0f, Delta);
+				}
 			}
 			bHandled = true;
 		}
diff --git a/Engine/Source/Editor/PhysicsAssetEditor/Private/PhysicsAssetEditorEditMode.h b/Engine/Source/Editor/PhysicsAssetEditor/Private/PhysicsAssetEditorEditMode.h
--- a/Engine/Source/Editor/PhysicsAssetEditor/Private/PhysicsAssetEditorEditMode.h
+++ b/Engine/Source/Editor/PhysicsAssetEditor/Private/PhysicsAssetEditorEditMode.h
@@ -52,13 +52,18 @@ public:
 	virtual bool IsCompatibleWith(FEditorModeID OtherModeID) const override { return true; }
 	virtual bool ReceivedFocus(FEditorViewportClient* ViewportClient, FViewport* Viewport);
 	virtual bool LostFocus(FEditorViewportClient* ViewportClient, FViewport* Viewport);
-	virtual bool RequiresLegacyViewportInteractions() const override { return true; }
+	virtual bool RequiresLegacyViewportInteractions() const override { return false; }
+	virtual void Enter() override;
+	virtual void Exit() override;
 
 	// Start IGizmoEdModeInterface overrides
 	virtual bool BeginTransform(const FGizmoState& InState) override;
 	virtual bool EndTransform(const FGizmoState& InState) override;
 	// End IGizmoEdModeInterface overrides
 	
+	/** Accessor for shared data, used by UPhysicsAssetSimInteraction. */
+	FPhysicsAssetEditorSharedData* GetSharedData() const { return SharedData; }
+
 private:
 	/** Simulation mouse forces */
 	bool SimMousePress(FEditorViewportClient* InViewportClient, FKey Key);
diff --git a/Engine/Source/Editor/PhysicsAssetEditor/Private/ViewportInteractions/PhysicsAssetSimInteraction.cpp b/Engine/Source/Editor/PhysicsAssetEditor/Private/ViewportInteractions/PhysicsAssetSimInteraction.cpp
new file mode 100644
--- /dev/null
+++ b/Engine/Source/Editor/PhysicsAssetEditor/Private/ViewportInteractions/PhysicsAssetSimInteraction.cpp
@@ -0,0 +1,272 @@
+// Copyright Epic Games, Inc. All Rights Reserved.
+
+#include "PhysicsAssetSimInteraction.h"
+
+#include "BaseBehaviors/MouseWheelBehavior.h"
+#include "PhysicsAssetEditorEditMode.h"
+#include "PhysicsAssetEditorSharedData.h"
+#include "PhysicsAssetEditorSkeletalMeshComponent.h"
+#include "PhysicsAssetEditorPhysicsHandleComponent.h"
+#include "Preferences/PhysicsAssetEditorOptions.h"
+#include "DrawDebugHelpers.h"
+#include "EditorViewportClient.h"
+#include "IPersonaPreviewScene.h"
+#include "PhysicsEngine/PhysicsAsset.h"
+#include "PhysicsEngine/SkeletalBodySetup.h"
+#include "SceneView.h"
+#include "ViewportInteractions/ViewportInteractionsBehaviorSource.h"
+
+UPhysicsAssetSimInteraction::UPhysicsAssetSimInteraction()
+{
+	using namespace UE::Editor::ViewportInteractions;
+
+	InteractionName = TEXT("PhysicsAssetSimulation");
+
+	if (ensure(ViewportClickDragBehavior.IsValid()))
+	{
+		ViewportClickDragBehavior->SetBindings({
+			FButtonBinding(EKeys::LeftMouseButton).TriggersStart().RequiredToStart(false).RequiredToContinue(false),
+			FButtonBinding(EKeys::RightMouseButton).TriggersStart().RequiredToStart(false).RequiredToContinue(false),
+			FButtonBinding(EKeys::LeftControl).RequiredToStart(true).RequiredToContinue(false)
+		});
+	}
+
+	// Scroll wheel for adjusting grab distance during an active grab.
+	// Gated on Ctrl being held via ModifierCheckFunc — the default zoom interaction
+	// rejects Ctrl+scroll, so there's no priority conflict.
+	UMouseWheelInputBehavior* MouseWheelBehavior = NewObject<UMouseWheelInputBehavior>();
+	MouseWheelBehavior->Initialize(this);
+	MouseWheelBehavior->SetDefaultPriority(FInputCapturePriority(
+		VIEWPORT_INTERACTIONS_DEFAULT_PRIORITY - 50));
+	MouseWheelBehavior->ModifierCheckFunc = [this](const FInputDeviceState& InputDeviceState)
+	{
+		// Only respond when we have a grabbed component (active Ctrl+RMB drag).
+		const FPhysicsAssetEditorSharedData* SharedData = GetSharedData();
+		return SharedData
+			&& SharedData->bRunningSimulation
+			&& SharedData->MouseHandle->GrabbedComponent != nullptr;
+	};
+
+	RegisterInputBehavior(MouseWheelBehavior);
+}
+
+bool UPhysicsAssetSimInteraction::CanBeActivated() const
+{
+	if (!Super::CanBeActivated())
+	{
+		return false;
+	}
+
+	const FPhysicsAssetEditorSharedData* SharedData = GetSharedData();
+	return SharedData && SharedData->bRunningSimulation;
+}
+
+void UPhysicsAssetSimInteraction::OnStateUpdated(const FInputDeviceState& InInputDeviceState)
+{
+	bIsRightMouseDown = InInputDeviceState.Mouse.Right.bDown;
+}
+
+void UPhysicsAssetSimInteraction::OnDragStart(const FInputDeviceRay& InDragStartPos)
+{
+	Super::OnDragStart(InDragStartPos);
+
+	FPhysicsAssetEditorSharedData* SharedData = GetSharedData();
+	FEditorViewportClient* ViewportClient = GetEditorViewportClient();
+	if (!SharedData || !ViewportClient)
+	{
+		return;
+	}
+
+	// Use the ray from the ITF input system for the line trace.
+	const FVector RayOrigin = InDragStartPos.WorldRay.Origin;
+	const FVector RayDirection = InDragStartPos.WorldRay.Direction;
+
+	FHitResult Result(1.f);
+	const bool bHit = SharedData->EditorSkelComp->LineTraceComponent(
+		Result, RayOrigin,
+		RayOrigin + RayDirection * SharedData->EditorOptions->InteractionDistance,
+		FCollisionQueryParams(NAME_None, true));
+
+	// Store click state on SharedData (used by DrawHUD).
+	SharedData->LastClickPos = FIntPoint(
+		static_cast<int32>(InDragStartPos.ScreenPosition.X),
+		static_cast<int32>(InDragStartPos.ScreenPosition.Y));
+	SharedData->LastClickOrigin = RayOrigin;
+	SharedData->LastClickDirection = RayDirection;
+	SharedData->bLastClickHit = bHit;
+	if (bHit)
+	{
+		SharedData->LastClickHitPos = Result.Location;
+		SharedData->LastClickHitNormal = Result.Normal;
+	}
+
+	bIsGrabbing = false;
+
+	if (!bHit)
+	{
+		return;
+	}
+
+	check(Result.Item != INDEX_NONE);
+	const FName BoneName = SharedData->PhysicsAsset->SkeletalBodySetups[Result.Item]->BoneName;
+
+	// Determine which mouse button triggered the drag.
+	if (bIsRightMouseDown)
+	{
+		// Ctrl+RMB: Grab body with physics handle.
+		SharedData->bManipulating = true;
+		bIsGrabbing = true;
+		DragX = 0.0f;
+		DragY = 0.0f;
+		SimGrabPush = 0.0f;
+
+		SharedData->MouseHandle->LinearDamping = SharedData->EditorOptions->HandleLinearDamping;
+		SharedData->MouseHandle->LinearStiffness = SharedData->EditorOptions->HandleLinearStiffness;
+		SharedData->MouseHandle->AngularDamping = SharedData->EditorOptions->HandleAngularDamping;
+		SharedData->MouseHandle->AngularStiffness = SharedData->EditorOptions->HandleAngularStiffness;
+		SharedData->MouseHandle->InterpolationSpeed = SharedData->EditorOptions->InterpolationSpeed;
+
+		SharedData->MouseHandle->GrabComponentAtLocationWithRotation(
+			SharedData->EditorSkelComp, BoneName, Result.Location, FRotator::ZeroRotator);
+
+		FViewport* Viewport = ViewportClient->Viewport;
+		FSceneViewFamilyContext ViewFamily(FSceneViewFamily::ConstructionValues(
+			Viewport, ViewportClient->GetScene(), ViewportClient->EngineShowFlags));
+		const FSceneView* View = ViewportClient->CalcSceneView(&ViewFamily);
+		const FMatrix InvViewMatrix = View->ViewMatrices.GetViewToWorld();
+
+		SimGrabMinPush = SimMinHoldDistance - (Result.Time * SharedData->EditorOptions->InteractionDistance);
+
+		SimGrabLocation = Result.Location;
+		SimGrabX = InvViewMatrix.GetUnitAxis(EAxis::X);
+		SimGrabY = InvViewMatrix.GetUnitAxis(EAxis::Y);
+		SimGrabZ = InvViewMatrix.GetUnitAxis(EAxis::Z);
+	}
+	else
+	{
+		// Ctrl+LMB: Poke body with impulse.
+		SharedData->EditorSkelComp->AddImpulseAtLocation(
+			RayDirection * SharedData->EditorOptions->PokeStrength,
+			Result.Location, BoneName);
+	}
+}
+
+void UPhysicsAssetSimInteraction::OnDragDelta(float InMouseDeltaX, float InMouseDeltaY)
+{
+	if (!bIsGrabbing)
+	{
+		return;
+	}
+
+	const FPhysicsAssetEditorSharedData* SharedData = GetSharedData();
+	FEditorViewportClient* ViewportClient = GetEditorViewportClient();
+	if (!SharedData || !ViewportClient || !SharedData->MouseHandle->GrabbedComponent)
+	{
+		return;
+	}
+
+	DragX = ViewportClient->Viewport->GetMouseX() - SharedData->LastClickPos.X;
+	DragY = ViewportClient->Viewport->GetMouseY() - SharedData->LastClickPos.Y;
+
+	FSceneViewFamilyContext ViewFamily(FSceneViewFamily::ConstructionValues(
+		ViewportClient->Viewport, ViewportClient->GetScene(),
+		ViewportClient->EngineShowFlags));
+	const FSceneView* View = ViewportClient->CalcSceneView(&ViewFamily);
+
+	const FVector4 ScreenOldPos = View->PixelToScreen(SharedData->LastClickPos.X, SharedData->LastClickPos.Y, 1.0f);
+	const FVector4 ScreenNewPos = View->PixelToScreen(DragX + SharedData->LastClickPos.X, DragY + SharedData->LastClickPos.Y, 1.0f);
+	const FVector4 ScreenDelta = ScreenNewPos - ScreenOldPos;
+	const FVector4 ProjectedDelta = View->ScreenToWorld(ScreenDelta);
+
+	const FVector LocalOffset = View->ViewMatrices.GetWorldToView().TransformPosition(SimGrabLocation + SimGrabZ * SimGrabPush);
+	const float ZDistance = ViewportClient->GetViewportType() == ELevelViewportType::LVT_Perspective ? FMath::Abs(LocalOffset.Z) : 1.0f;
+	const FVector WorldDelta = FVector(ProjectedDelta) * ZDistance;
+
+	const FVector NewLocation = SimGrabLocation + WorldDelta + SimGrabZ * SimGrabPush;
+
+	float QuickRadius = 5.0f - SimGrabPush / SimHoldDistanceChangeDelta;
+	QuickRadius = QuickRadius < 2.0f ? 2.0f : QuickRadius;
+
+	DrawDebugPoint(ViewportClient->GetWorld(), NewLocation, QuickRadius, FColorList::Red, false, 0.3f);
+
+	SharedData->MouseHandle->SetTargetLocation(NewLocation);
+	SharedData->MouseHandle->GrabbedComponent->WakeRigidBody(SharedData->MouseHandle->GrabbedBoneName);
+}
+
+void UPhysicsAssetSimInteraction::OnDragEnd(const FInputDeviceRay& InDragEndPos)
+{
+	Super::OnDragEnd(InDragEndPos);
+
+	if (!bIsGrabbing)
+	{
+		return;
+	}
+
+	if (FPhysicsAssetEditorSharedData* SharedData = GetSharedData())
+	{
+		SharedData->bManipulating = false;
+
+		if (SharedData->MouseHandle->GrabbedComponent)
+		{
+			SharedData->MouseHandle->GrabbedComponent->WakeRigidBody(SharedData->MouseHandle->GrabbedBoneName);
+			SharedData->MouseHandle->ReleaseComponent();
+		}
+	}
+
+	bIsGrabbing = false;
+}
+
+FInputRayHit UPhysicsAssetSimInteraction::ShouldRespondToMouseWheel(const FInputDeviceRay& CurrentPos)
+{
+	const FPhysicsAssetEditorSharedData* SharedData = GetSharedData();
+	if (SharedData && SharedData->bRunningSimulation && SharedData->MouseHandle->GrabbedComponent)
+	{
+		return FInputRayHit(0.0f);
+	}
+
+	return FInputRayHit();
+}
+
+void UPhysicsAssetSimInteraction::OnMouseWheelScrollUp(const FInputDeviceRay& CurrentPos)
+{
+	const FPhysicsAssetEditorSharedData* SharedData = GetSharedData();
+	if (!SharedData || !SharedData->MouseHandle->GrabbedComponent)
+	{
+		return;
+	}
+
+	SimGrabPush += SimHoldDistanceChangeDelta;
+
+	// Trigger a zero-delta drag update to reposition the grab target.
+	OnDragDelta(0.0f, 0.0f);
+}
+
+void UPhysicsAssetSimInteraction::OnMouseWheelScrollDown(const FInputDeviceRay& CurrentPos)
+{
+	FPhysicsAssetEditorSharedData* SharedData = GetSharedData();
+	if (!SharedData || !SharedData->MouseHandle->GrabbedComponent)
+	{
+		return;
+	}
+
+	SimGrabPush -= SimHoldDistanceChangeDelta;
+	SimGrabPush = FMath::Max(SimGrabMinPush, SimGrabPush);
+
+	// Trigger a zero-delta drag update to reposition the grab target.
+	OnDragDelta(0.0f, 0.0f);
+}
+
+void UPhysicsAssetSimInteraction::SetEditMode(FPhysicsAssetEditorEditMode* InEditMode)
+{
+	EditMode = InEditMode;
+}
+
+FPhysicsAssetEditorSharedData* UPhysicsAssetSimInteraction::GetSharedData() const
+{
+	if (EditMode)
+	{
+		return EditMode->GetSharedData();
+	}
+
+	return nullptr;
+}
diff --git a/Engine/Source/Editor/PhysicsAssetEditor/Private/ViewportInteractions/PhysicsAssetSimInteraction.h b/Engine/Source/Editor/PhysicsAssetEditor/Private/ViewportInteractions/PhysicsAssetSimInteraction.h
new file mode 100644
--- /dev/null
+++ b/Engine/Source/Editor/PhysicsAssetEditor/Private/ViewportInteractions/PhysicsAssetSimInteraction.h
@@ -0,0 +1,67 @@
+// Copyright Epic Games, Inc. All Rights Reserved.
+
+#pragma once
+
+#include "BaseBehaviors/BehaviorTargetInterfaces.h"
+#include "ViewportInteractions/ViewportDragInteraction.h"
+#include "PhysicsAssetSimInteraction.generated.h"
+
+class FPhysicsAssetEditorEditMode;
+class FPhysicsAssetEditorSharedData;
+
+/**
+ * Viewport interaction for the Physics Asset Editor simulation mode.
+ * Handles Ctrl+RMB grab (drag body with physics handle), Ctrl+LMB poke (apply impulse),
+ * and Ctrl+scroll wheel (adjust grab distance).
+ */
+UCLASS(Transient)
+class UPhysicsAssetSimInteraction
+	: public UViewportDragInteraction
+	, public IMouseWheelBehaviorTarget
+{
+	GENERATED_BODY()
+
+public:
+	UPhysicsAssetSimInteraction();
+
+	virtual bool CanBeActivated() const override;
+	virtual void OnDragStart(const FInputDeviceRay& InDragStartPos) override;
+	virtual void OnDragDelta(float InMouseDeltaX, float InMouseDeltaY) override;
+	virtual void OnDragEnd(const FInputDeviceRay& InDragEndPos) override;
+	virtual void OnStateUpdated(const FInputDeviceState& InInputDeviceState) override;
+
+	//~ Begin IMouseWheelBehaviorTarget
+	virtual FInputRayHit ShouldRespondToMouseWheel(const FInputDeviceRay& CurrentPos) override;
+	virtual void OnMouseWheelScrollUp(const FInputDeviceRay& CurrentPos) override;
+	virtual void OnMouseWheelScrollDown(const FInputDeviceRay& CurrentPos) override;
+	//~ End IMouseWheelBehaviorTarget
+
+	/** Set by the edit mode during registration. */
+	void SetEditMode(FPhysicsAssetEditorEditMode* InEditMode);
+
+private:
+	FPhysicsAssetEditorSharedData* GetSharedData() const;
+
+private:
+	FPhysicsAssetEditorEditMode* EditMode = nullptr;
+
+	/** Whether RMB is currently held, determines drag vs pulse. */
+	bool bIsRightMouseDown = false;
+
+	/** Simulation grab state */
+	float SimGrabPush = 0.0f;
+	float SimGrabMinPush = 0.0f;
+	FVector SimGrabLocation = FVector::ZeroVector;
+	FVector SimGrabX = FVector::ZeroVector;
+	FVector SimGrabY = FVector::ZeroVector;
+	FVector SimGrabZ = FVector::ZeroVector;
+
+	float DragX = 0.0f;
+	float DragY = 0.0f;
+
+	/** Whether the current drag is a grab (RMB) vs poke (LMB). */
+	bool bIsGrabbing = false;
+
+	static constexpr float SimHoldDistanceChangeDelta = 20.0f;
+	static constexpr float SimMinHoldDistance = 10.0f;
+};
