PhysicsAssetEditor camera interactions work in legacy path

#Synced-CL 52513890

diff --git a/Engine/Source/Editor/PhysicsAssetEditor/Private/PhysicsAssetEditor.cpp b/Engine/Source/Editor/PhysicsAssetEditor/Private/PhysicsAssetEditor.cpp
--- a/Engine/Source/Editor/PhysicsAssetEditor/Private/PhysicsAssetEditor.cpp
+++ b/Engine/Source/Editor/PhysicsAssetEditor/Private/PhysicsAssetEditor.cpp
@@ -305,16 +305,16 @@ void FPhysicsAssetEditor::InitPhysicsAssetEditor(const EToolkitMode::Type Mode,
 	const bool bCreateDefaultToolbar = true;
 	FAssetEditorToolkit::InitAssetEditor(Mode, InitToolkitHost, PhysicsAssetEditorAppIdentifier, FTabManager::FLayout::NullLayout, bCreateDefaultStandaloneMenu, bCreateDefaultToolbar, ObjectToEdit);
 
-	PhysicsAssetEditor::CVarUseITFViewportInteractions->OnChangedDelegate().AddSPLambda(this, [ModeTools = TWeakPtr<FEditorModeTools>(GetEditorModeManager().AsShared())](IConsoleVariable*)
+	PhysicsAssetEditor::CVarUseITFViewportInteractions->OnChangedDelegate().AddSPLambda(this, [WeakModeTools = TWeakPtr<FEditorModeTools>(GetEditorModeManager().AsShared())](IConsoleVariable*)
 	{
-		if (TSharedPtr<FEditorModeTools> ModeToolsPin = ModeTools.Pin())
+		if (TSharedPtr<FEditorModeTools> StrongModeTools = WeakModeTools.Pin())
 		{
 			const bool bEnabled = PhysicsAssetEditor::CVarUseITFViewportInteractions->GetBool();
-			if (UModeManagerInteractiveToolsContext* ToolsContext = ModeToolsPin->GetInteractiveToolsContext())
+			if (UModeManagerInteractiveToolsContext* ToolsContext = StrongModeTools->GetInteractiveToolsContext())
 			{
 				ToolsContext->SetViewportInteractionsEnabled(bEnabled);
 			}
-			ModeToolsPin->SetSupportsViewportITF(bEnabled);
+			StrongModeTools->SetSupportsViewportITF(bEnabled);
 		}
 	});
 
diff --git a/Engine/Source/Editor/PhysicsAssetEditor/Private/PhysicsAssetEditorEditMode.cpp b/Engine/Source/Editor/PhysicsAssetEditor/Private/PhysicsAssetEditorEditMode.cpp
--- a/Engine/Source/Editor/PhysicsAssetEditor/Private/PhysicsAssetEditorEditMode.cpp
+++ b/Engine/Source/Editor/PhysicsAssetEditor/Private/PhysicsAssetEditorEditMode.cpp
@@ -59,7 +59,6 @@ void FPhysicsAssetEditorEditMode::Enter()
 
 	if (UModeManagerInteractiveToolsContext* Context = GetModeManager()->GetInteractiveToolsContext())
 	{
-		Context->SetViewportInteractionsEnabled(true);
 		Context->OnBuildViewportInteractions().AddSPLambda(this, [this](UViewportInteractionsBehaviorSource* Source)
 		{
 			UPhysicsAssetSimInteraction* NewSimInteraction = Source->AddInteraction<UPhysicsAssetSimInteraction>(true);
